UE4 - ES3.1 fixes

diff --git a/Engine/Source/Runtime/EmptyRHI/Private/EmptyRHI.cpp b/Engine/Source/Runtime/EmptyRHI/Private/EmptyRHI.cpp
--- a/Engine/Source/Runtime/EmptyRHI/Private/EmptyRHI.cpp
+++ b/Engine/Source/Runtime/EmptyRHI/Private/EmptyRHI.cpp
@@ -49,7 +49,7 @@ FEmptyDynamicRHI::FEmptyDynamicRHI()
 // 	GMaxTextureArrayLayers = 8192;
 
 //	GShaderPlatformForFeatureLevel[ERHIFeatureLevel::ES2] = SP_NumPlatforms;
-//	GShaderPlatformForFeatureLevel[ERHIFeatureLevel::SM3] = SP_NumPlatforms;
+//	GShaderPlatformForFeatureLevel[ERHIFeatureLevel::ES3_1] = SP_NumPlatforms;
 //	GShaderPlatformForFeatureLevel[ERHIFeatureLevel::SM4] = SP_NumPlatforms;
 //	GShaderPlatformForFeatureLevel[ERHIFeatureLevel::SM5] = ;
 
diff --git a/Engine/Source/Runtime/Engine/Private/Landscape/LandscapeRender.cpp b/Engine/Source/Runtime/Engine/Private/Landscape/LandscapeRender.cpp
--- a/Engine/Source/Runtime/Engine/Private/Landscape/LandscapeRender.cpp
+++ b/Engine/Source/Runtime/Engine/Private/Landscape/LandscapeRender.cpp
@@ -2169,7 +2169,7 @@ FLandscapeSharedBuffers::FLandscapeSharedBuffers(int32 InSharedBuffersKey, int32
 ,	AdjacencyIndexBuffers(NULL)
 ,	bUse32BitIndices(false)
 {
-	if (InFeatureLevel > ERHIFeatureLevel::ES2)
+	if (InFeatureLevel > ERHIFeatureLevel::ES3_1)
 	{
 		// Vertex Buffer cannot be shared
 		VertexBuffer = new FLandscapeVertexBuffer(SubsectionSizeVerts, NumSubsections);
diff --git a/Engine/Source/Runtime/Engine/Private/UnrealEngine.cpp b/Engine/Source/Runtime/Engine/Private/UnrealEngine.cpp
--- a/Engine/Source/Runtime/Engine/Private/UnrealEngine.cpp
+++ b/Engine/Source/Runtime/Engine/Private/UnrealEngine.cpp
@@ -10041,7 +10041,7 @@ static TAutoConsoleVariable<int32> CVarAllowHighQualityLightMaps(
 bool AllowHighQualityLightmaps(ERHIFeatureLevel::Type FeatureLevel)
 {
 	return FPlatformProperties::SupportsHighQualityLightmaps()
-		&& (FeatureLevel > ERHIFeatureLevel::ES2)
+		&& (FeatureLevel > ERHIFeatureLevel::ES3_1)
 		&& (CVarAllowHighQualityLightMaps.GetValueOnAnyThread() != 0);
 }
 
