UE4 - Some fixes for EmptyRHI

diff --git a/Engine/Source/Runtime/EmptyRHI/Private/EmptyShaders.cpp b/Engine/Source/Runtime/EmptyRHI/Private/EmptyShaders.cpp
--- a/Engine/Source/Runtime/EmptyRHI/Private/EmptyShaders.cpp
+++ b/Engine/Source/Runtime/EmptyRHI/Private/EmptyShaders.cpp
@@ -76,12 +76,12 @@ FEmptyBoundShaderState::FEmptyBoundShaderState(
 			FGeometryShaderRHIParamRef InGeometryShaderRHI)
 	:	CacheLink(InVertexDeclarationRHI,InVertexShaderRHI,InPixelShaderRHI,InHullShaderRHI,InDomainShaderRHI,InGeometryShaderRHI,this)
 {
-	FEmptyVertexDeclaration* InVertexDeclaration = ResourceCast(InVertexDeclarationRHI);
-	FEmptyVertexShader* InVertexShader = ResourceCast(InVertexShaderRHI);
-	FEmptyPixelShader* InPixelShader = ResourceCast(InPixelShaderRHI);
-	FEmptyHullShader* InHullShader = ResourceCast(InHullShaderRHI);
-	FEmptyDomainShader* InDomainShader = ResourceCast(InDomainShaderRHI);
-	FEmptyGeometryShader* InGeometryShader = ResourceCast(InGeometryShaderRHI);
+	FEmptyVertexDeclaration* InVertexDeclaration = FEmptyDynamicRHI::ResourceCast(InVertexDeclarationRHI);
+	FEmptyVertexShader* InVertexShader = FEmptyDynamicRHI::ResourceCast(InVertexShaderRHI);
+	FEmptyPixelShader* InPixelShader = FEmptyDynamicRHI::ResourceCast(InPixelShaderRHI);
+	FEmptyHullShader* InHullShader = FEmptyDynamicRHI::ResourceCast(InHullShaderRHI);
+	FEmptyDomainShader* InDomainShader = FEmptyDynamicRHI::ResourceCast(InDomainShaderRHI);
+	FEmptyGeometryShader* InGeometryShader = FEmptyDynamicRHI::ResourceCast(InGeometryShaderRHI);
 
 	// cache everything
 	VertexDeclaration = InVertexDeclaration;
diff --git a/Engine/Source/Runtime/EmptyRHI/Private/EmptyStructuredBuffer.cpp b/Engine/Source/Runtime/EmptyRHI/Private/EmptyStructuredBuffer.cpp
--- a/Engine/Source/Runtime/EmptyRHI/Private/EmptyStructuredBuffer.cpp
+++ b/Engine/Source/Runtime/EmptyRHI/Private/EmptyStructuredBuffer.cpp
@@ -25,16 +25,16 @@ FEmptyStructuredBuffer::~FEmptyStructuredBuffer()
 
 
 
-FStructuredBufferRHIRef FEmptyDynamicRHI::RHICreateStructuredBuffer(uint32 Stride,uint32 Size,FResourceArrayInterface* ResourceArray,uint32 InUsage)
+FStructuredBufferRHIRef FEmptyDynamicRHI::RHICreateStructuredBuffer(uint32 Stride, uint32 Size, uint32 InUsage, FRHIResourceCreateInfo& CreateInfo)
 {
-	return new FEmptyStructuredBuffer(Stride, Size, ResourceArray, InUsage);
+	return nullptr;//new FEmptyStructuredBuffer(Stride, Size, ResourceArray, InUsage);
 }
 
 void* FEmptyDynamicRHI::RHILockStructuredBuffer(FStructuredBufferRHIParamRef StructuredBufferRHI,uint32 Offset,uint32 Size,EResourceLockMode LockMode)
 {
 	FEmptyStructuredBuffer* StructuredBuffer = ResourceCast(StructuredBufferRHI);
 
-	return NULL;
+	return nullptr;
 }
 
 void FEmptyDynamicRHI::RHIUnlockStructuredBuffer(FStructuredBufferRHIParamRef StructuredBufferRHI)
diff --git a/Engine/Source/Runtime/EmptyRHI/Private/EmptyTexture.cpp b/Engine/Source/Runtime/EmptyRHI/Private/EmptyTexture.cpp
--- a/Engine/Source/Runtime/EmptyRHI/Private/EmptyTexture.cpp
+++ b/Engine/Source/Runtime/EmptyRHI/Private/EmptyTexture.cpp
@@ -191,7 +191,7 @@ void FEmptyDynamicRHI::RHIUnlockTexture2DArray(FTexture2DArrayRHIParamRef Textur
 
 void FEmptyDynamicRHI::RHIUpdateTexture2D(FTexture2DRHIParamRef TextureRHI, uint32 MipIndex, const struct FUpdateTextureRegion2D& UpdateRegion, uint32 SourcePitch, const uint8* SourceData)
 {	
-	FEmptyTexture3D* Texture = ResourceCast(TextureRHI);
+	FEmptyTexture2D* Texture = ResourceCast(TextureRHI);
 
 }
 
diff --git a/Engine/Source/Runtime/EmptyRHI/Private/EmptyUniformBuffer.cpp b/Engine/Source/Runtime/EmptyRHI/Private/EmptyUniformBuffer.cpp
--- a/Engine/Source/Runtime/EmptyRHI/Private/EmptyUniformBuffer.cpp
+++ b/Engine/Source/Runtime/EmptyRHI/Private/EmptyUniformBuffer.cpp
@@ -54,8 +54,8 @@ void FEmptyGlobalUniformBuffer::UpdateConstant(const uint8* Data, uint16 Offset,
 	Uniform buffer RHI object
 -----------------------------------------------------------------------------*/
 
-FEmptyUniformBuffer::FEmptyUniformBuffer(const void* Contents, uint32 NumBytes, EUniformBufferUsage Usage)
-	: FRHIUniformBuffer(NumBytes)
+FEmptyUniformBuffer::FEmptyUniformBuffer(const FRHIUniformBufferLayout& InLayout, const void* Contents, uint32 NumBytes, EUniformBufferUsage Usage)
+	: FRHIUniformBuffer(InLayout)
 {
 }
 
@@ -63,8 +63,8 @@ FEmptyUniformBuffer::~FEmptyUniformBuffer()
 {
 }
 
-FUniformBufferRHIRef FEmptyDynamicRHI::RHICreateUniformBuffer(const void* Contents, uint32 NumBytes, EUniformBufferUsage Usage)
+FUniformBufferRHIRef FEmptyDynamicRHI::RHICreateUniformBuffer(const void* Contents, const FRHIUniformBufferLayout& Layout, EUniformBufferUsage Usage)
 {
 	check(IsInRenderingThread());
-	return new FEmptyUniformBuffer(Contents, NumBytes, Usage);
+	return nullptr;//new FEmptyUniformBuffer(Contents, NumBytes, Usage);
 }
diff --git a/Engine/Source/Runtime/EmptyRHI/Private/EmptyViewport.cpp b/Engine/Source/Runtime/EmptyRHI/Private/EmptyViewport.cpp
--- a/Engine/Source/Runtime/EmptyRHI/Private/EmptyViewport.cpp
+++ b/Engine/Source/Runtime/EmptyRHI/Private/EmptyViewport.cpp
@@ -47,7 +47,7 @@ void FEmptyDynamicRHI::RHIBeginDrawingViewport(FViewportRHIParamRef ViewportRHI,
 {
 	FEmptyViewport* Viewport = ResourceCast(ViewportRHI);
 
-	RHISetRenderTarget(RHIGetViewportBackBuffer(ViewportRHI), NULL);
+	//RHISetRenderTarget(RHIGetViewportBackBuffer(ViewportRHI), NULL);
 }
 
 void FEmptyDynamicRHI::RHIEndDrawingViewport(FViewportRHIParamRef ViewportRHI,bool bPresent,bool bLockToVsync)
diff --git a/Engine/Source/Runtime/EmptyRHI/Public/EmptyResources.h b/Engine/Source/Runtime/EmptyRHI/Public/EmptyResources.h
--- a/Engine/Source/Runtime/EmptyRHI/Public/EmptyResources.h
+++ b/Engine/Source/Runtime/EmptyRHI/Public/EmptyResources.h
@@ -21,47 +21,6 @@ public:
 };
 
 
-
-
-/** This represents a vertex shader that hasn't been combined with a specific declaration to create a bound shader. */
-template<typename BaseResourceType /*, typename ShaderType */>
-class TEmptyBaseShader : public BaseResourceType, public IRefCountedObject
-{
-public:
-
-	/** Initialization constructor. */
-	TEmptyBaseShader()
-	{
-	}
-	TEmptyBaseShader(const TArray<uint8>& InCode);
-
-	/** Destructor */
-	virtual ~TEmptyBaseShader();
-
-
-	// IRefCountedObject interface.
-	virtual uint32 AddRef() const
-	{
-		return FRHIResource::AddRef();
-	}
-	virtual uint32 Release() const
-	{
-		return FRHIResource::Release();
-	}
-	virtual uint32 GetRefCount() const
-	{
-		return FRHIResource::GetRefCount();
-	}
-};
-
-typedef TEmptyBaseShader<FRHIVertexShader> FEmptyVertexShader;
-typedef TEmptyBaseShader<FRHIPixelShader> FEmptyPixelShader;
-typedef TEmptyBaseShader<FRHIHullShader> FEmptyHullShader;
-typedef TEmptyBaseShader<FRHIDomainShader> FEmptyDomainShader;
-typedef TEmptyBaseShader<FRHIComputeShader> FEmptyComputeShader;
-typedef TEmptyBaseShader<FRHIGeometryShader> FEmptyGeometryShader;
-
-
 /**
  * Combined shader state and vertex definition for rendering geometry. 
  * Each unique instance consists of a vertex decl, vertex shader, and pixel shader.
@@ -263,7 +222,7 @@ class FEmptyUniformBuffer : public FRHIUniformBuffer
 public:
 
 	// Constructor
-	FEmptyUniformBuffer(const void* Contents, uint32 NumBytes, EUniformBufferUsage Usage);
+	FEmptyUniformBuffer(const FRHIUniformBufferLayout& InLayout, const void* Contents, uint32 NumBytes, EUniformBufferUsage Usage);
 
 	// Destructor 
 	~FEmptyUniformBuffer();
diff --git a/Engine/Source/Runtime/EmptyRHI/Public/EmptyShaderResources.h b/Engine/Source/Runtime/EmptyRHI/Public/EmptyShaderResources.h
new file mode 100644
--- /dev/null
+++ b/Engine/Source/Runtime/EmptyRHI/Public/EmptyShaderResources.h
@@ -0,0 +1,47 @@
+// Copyright 1998-2015 Epic Games, Inc. All Rights Reserved.
+
+/*=============================================================================
+	EmptyShaderResources.h: Empty shader resource RHI definitions.
+=============================================================================*/
+
+#pragma once
+
+#include "CrossCompilerCommon.h"
+
+/** This represents a vertex shader that hasn't been combined with a specific declaration to create a bound shader. */
+template<typename BaseResourceType /*, typename ShaderType */>
+class TEmptyBaseShader : public BaseResourceType, public IRefCountedObject
+{
+public:
+
+	/** Initialization constructor. */
+	TEmptyBaseShader()
+	{
+	}
+	TEmptyBaseShader(const TArray<uint8>& InCode);
+
+	/** Destructor */
+	virtual ~TEmptyBaseShader();
+
+
+	// IRefCountedObject interface.
+	virtual uint32 AddRef() const
+	{
+		return FRHIResource::AddRef();
+	}
+	virtual uint32 Release() const
+	{
+		return FRHIResource::Release();
+	}
+	virtual uint32 GetRefCount() const
+	{
+		return FRHIResource::GetRefCount();
+	}
+};
+
+typedef TEmptyBaseShader<FRHIVertexShader> FEmptyVertexShader;
+typedef TEmptyBaseShader<FRHIPixelShader> FEmptyPixelShader;
+typedef TEmptyBaseShader<FRHIHullShader> FEmptyHullShader;
+typedef TEmptyBaseShader<FRHIDomainShader> FEmptyDomainShader;
+typedef TEmptyBaseShader<FRHIComputeShader> FEmptyComputeShader;
+typedef TEmptyBaseShader<FRHIGeometryShader> FEmptyGeometryShader;
