UE4 - More missing empty

diff --git a/Engine/Source/Runtime/EmptyRHI/Private/EmptyCommands.cpp b/Engine/Source/Runtime/EmptyRHI/Private/EmptyCommands.cpp
--- a/Engine/Source/Runtime/EmptyRHI/Private/EmptyCommands.cpp
+++ b/Engine/Source/Runtime/EmptyRHI/Private/EmptyCommands.cpp
@@ -396,6 +396,10 @@ void FEmptyDynamicRHI::RHIClearMRT(bool bClearColor,int32 NumClearColors,const F
 
 }
 
+void FEmptyDynamicRHI::RHIBindClearMRTValues(bool bClearColor, int32 NumClearColors, const FLinearColor* ColorArray, bool bClearDepth, float Depth, bool bClearStencil, uint32 Stencil)
+{
+}
+
 
 
 void FEmptyDynamicRHI::RHISuspendRendering()
@@ -439,3 +443,10 @@ void* FEmptyDynamicRHI::RHIGetNativeDevice()
 	return nullptr;
 }
 
+void FEmptyDynamicRHI::RHIExecuteCommandList(FRHICommandList* CmdList)
+{
+}
+
+void FEmptyDynamicRHI::RHIEnableDepthBoundsTest(bool bEnable, float MinDepth, float MaxDepth)
+{
+}
diff --git a/Engine/Source/Runtime/EmptyRHI/Private/EmptyQuery.cpp b/Engine/Source/Runtime/EmptyRHI/Private/EmptyQuery.cpp
--- a/Engine/Source/Runtime/EmptyRHI/Private/EmptyQuery.cpp
+++ b/Engine/Source/Runtime/EmptyRHI/Private/EmptyQuery.cpp
@@ -45,3 +45,15 @@ bool FEmptyDynamicRHI::RHIGetRenderQueryResult(FRenderQueryRHIParamRef QueryRHI,
 
 	return false;
 }
+
+void FEmptyDynamicRHI::RHIBeginOcclusionQueryBatch()
+{
+}
+
+void FEmptyDynamicRHI::RHIEndOcclusionQueryBatch()
+{
+}
+
+void FEmptyDynamicRHI::RHISubmitCommandsHint()
+{
+}
diff --git a/Engine/Source/Runtime/EmptyRHI/Private/EmptyTexture.cpp b/Engine/Source/Runtime/EmptyRHI/Private/EmptyTexture.cpp
--- a/Engine/Source/Runtime/EmptyRHI/Private/EmptyTexture.cpp
+++ b/Engine/Source/Runtime/EmptyRHI/Private/EmptyTexture.cpp
@@ -241,3 +241,27 @@ void FEmptyDynamicRHI::RHIVirtualTextureSetFirstMipInMemory(FTexture2DRHIParamRe
 void FEmptyDynamicRHI::RHIVirtualTextureSetFirstMipVisible(FTexture2DRHIParamRef TextureRHI, uint32 FirstMip)
 {
 }
+
+uint64 FEmptyDynamicRHI::RHICalcTexture2DPlatformSize(uint32 SizeX, uint32 SizeY, uint8 Format, uint32 NumMips, uint32 NumSamples, uint32 Flags, uint32& OutAlign)
+{
+	return 0;
+}
+
+uint64 FEmptyDynamicRHI::RHICalcTexture3DPlatformSize(uint32 SizeX, uint32 SizeY, uint32 SizeZ, uint8 Format, uint32 NumMips, uint32 Flags, uint32& OutAlign)
+{
+	return 0;
+}
+
+uint64 FEmptyDynamicRHI::RHICalcTextureCubePlatformSize(uint32 Size, uint8 Format, uint32 NumMips, uint32 Flags, uint32& OutAlign)
+{
+	return 0;
+}
+
+FTextureReferenceRHIRef FEmptyDynamicRHI::RHICreateTextureReference(FLastRenderTimeContainer* LastRenderTime)
+{
+	return nullptr;
+}
+
+void FEmptyDynamicRHI::RHIUpdateTextureReference(FTextureReferenceRHIParamRef TextureRef, FTextureRHIParamRef NewTexture)
+{
+}
