Added ability to create a texture UAV to a specific mip and create UAVs for additional texture types.

diff --git a/Engine/Source/Runtime/EmptyRHI/Private/EmptyUAV.cpp b/Engine/Source/Runtime/EmptyRHI/Private/EmptyUAV.cpp
--- a/Engine/Source/Runtime/EmptyRHI/Private/EmptyUAV.cpp
+++ b/Engine/Source/Runtime/EmptyRHI/Private/EmptyUAV.cpp
@@ -22,7 +22,7 @@ FUnorderedAccessViewRHIRef FEmptyDynamicRHI::RHICreateUnorderedAccessView(FStruc
 	return UAV;
 }
 
-FUnorderedAccessViewRHIRef FEmptyDynamicRHI::RHICreateUnorderedAccessView(FTextureRHIParamRef TextureRHI)
+FUnorderedAccessViewRHIRef FEmptyDynamicRHI::RHICreateUnorderedAccessView(FTextureRHIParamRef TextureRHI, uint32 MipLevel)
 {
 	FEmptySurface& Surface = GetEmptySurfaceFromRHITexture(TextureRHI);
 
diff --git a/Engine/Source/Runtime/EmptyRHI/Public/EmptyRHI.h b/Engine/Source/Runtime/EmptyRHI/Public/EmptyRHI.h
--- a/Engine/Source/Runtime/EmptyRHI/Public/EmptyRHI.h
+++ b/Engine/Source/Runtime/EmptyRHI/Public/EmptyRHI.h
@@ -62,7 +62,7 @@ public:
 	virtual void* RHILockStructuredBuffer(FStructuredBufferRHIParamRef StructuredBuffer, uint32 Offset, uint32 SizeRHI, EResourceLockMode LockMode) final override;
 	virtual void RHIUnlockStructuredBuffer(FStructuredBufferRHIParamRef StructuredBuffer) final override;
 	virtual FUnorderedAccessViewRHIRef RHICreateUnorderedAccessView(FStructuredBufferRHIParamRef StructuredBuffer, bool bUseUAVCounter, bool bAppendBuffer) final override;
-	virtual FUnorderedAccessViewRHIRef RHICreateUnorderedAccessView(FTextureRHIParamRef Texture) final override;
+	virtual FUnorderedAccessViewRHIRef RHICreateUnorderedAccessView(FTextureRHIParamRef Texture, uint32 MipLevel) final override;
 	virtual FUnorderedAccessViewRHIRef RHICreateUnorderedAccessView(FVertexBufferRHIParamRef VertexBuffer, uint8 Format) final override;
 	virtual FShaderResourceViewRHIRef RHICreateShaderResourceView(FStructuredBufferRHIParamRef StructuredBuffer) final override;
 	virtual FShaderResourceViewRHIRef RHICreateShaderResourceView(FVertexBufferRHIParamRef VertexBuffer, uint32 Stride, uint8 Format) final override;
@@ -193,7 +193,7 @@ public:
 	virtual void RHISetDepthStencilState(FDepthStencilStateRHIParamRef NewState, uint32 StencilRef) final override;
 	virtual void RHISetBlendState(FBlendStateRHIParamRef NewState, const FLinearColor& BlendFactor) final override;
 	virtual void RHISetRenderTargets(uint32 NumSimultaneousRenderTargets, const FRHIRenderTargetView* NewRenderTargets, const FRHIDepthRenderTargetView* NewDepthStencilTarget, uint32 NumUAVs, const FUnorderedAccessViewRHIParamRef* UAVs) final override;
-	virtual void RHISetRenderTargetsAndClear(const FRHISetRenderTargetsInfo& RenderTargetsInfo) final override;	
+	virtual void RHISetRenderTargetsAndClear(const FRHISetRenderTargetsInfo& RenderTargetsInfo) final override;
 	virtual void RHIDrawPrimitive(uint32 PrimitiveType, uint32 BaseVertexIndex, uint32 NumPrimitives, uint32 NumInstances) final override;
 	virtual void RHIDrawPrimitiveIndirect(uint32 PrimitiveType, FVertexBufferRHIParamRef ArgumentBuffer, uint32 ArgumentOffset) final override;
 	virtual void RHIDrawIndexedIndirect(FIndexBufferRHIParamRef IndexBufferRHI, uint32 PrimitiveType, FStructuredBufferRHIParamRef ArgumentsBufferRHI, int32 DrawArgumentsIndex, uint32 NumInstances) final override;
diff --git a/Engine/Source/Runtime/IOS/MetalRHI/Private/MetalUAV.cpp b/Engine/Source/Runtime/IOS/MetalRHI/Private/MetalUAV.cpp
--- a/Engine/Source/Runtime/IOS/MetalRHI/Private/MetalUAV.cpp
+++ b/Engine/Source/Runtime/IOS/MetalRHI/Private/MetalUAV.cpp
@@ -50,7 +50,7 @@ FUnorderedAccessViewRHIRef FMetalDynamicRHI::RHICreateUnorderedAccessView(FStruc
 	return UAV;
 }
 
-FUnorderedAccessViewRHIRef FMetalDynamicRHI::RHICreateUnorderedAccessView(FTextureRHIParamRef TextureRHI)
+FUnorderedAccessViewRHIRef FMetalDynamicRHI::RHICreateUnorderedAccessView(FTextureRHIParamRef TextureRHI, uint32 MipLevel)
 {
 	FMetalSurface* Surface = GetMetalSurfaceFromRHITexture(TextureRHI);
 
diff --git a/Engine/Source/Runtime/IOS/MetalRHI/Public/MetalRHI.h b/Engine/Source/Runtime/IOS/MetalRHI/Public/MetalRHI.h
--- a/Engine/Source/Runtime/IOS/MetalRHI/Public/MetalRHI.h
+++ b/Engine/Source/Runtime/IOS/MetalRHI/Public/MetalRHI.h
@@ -65,7 +65,7 @@ public:
 	virtual void* RHILockStructuredBuffer(FStructuredBufferRHIParamRef StructuredBuffer, uint32 Offset, uint32 SizeRHI, EResourceLockMode LockMode) final override;
 	virtual void RHIUnlockStructuredBuffer(FStructuredBufferRHIParamRef StructuredBuffer) final override;
 	virtual FUnorderedAccessViewRHIRef RHICreateUnorderedAccessView(FStructuredBufferRHIParamRef StructuredBuffer, bool bUseUAVCounter, bool bAppendBuffer) final override;
-	virtual FUnorderedAccessViewRHIRef RHICreateUnorderedAccessView(FTextureRHIParamRef Texture) final override;
+	virtual FUnorderedAccessViewRHIRef RHICreateUnorderedAccessView(FTextureRHIParamRef Texture, uint32 MipLevel) final override;
 	virtual FUnorderedAccessViewRHIRef RHICreateUnorderedAccessView(FVertexBufferRHIParamRef VertexBuffer, uint8 Format) final override;
 	virtual FShaderResourceViewRHIRef RHICreateShaderResourceView(FStructuredBufferRHIParamRef StructuredBuffer) final override;
 	virtual FShaderResourceViewRHIRef RHICreateShaderResourceView(FVertexBufferRHIParamRef VertexBuffer, uint32 Stride, uint8 Format) final override;
diff --git a/Engine/Source/Runtime/NullDrv/Public/NullRHI.h b/Engine/Source/Runtime/NullDrv/Public/NullRHI.h
--- a/Engine/Source/Runtime/NullDrv/Public/NullRHI.h
+++ b/Engine/Source/Runtime/NullDrv/Public/NullRHI.h
@@ -187,7 +187,7 @@ public:
 	}
 
 
-	virtual FUnorderedAccessViewRHIRef RHICreateUnorderedAccessView(FTextureRHIParamRef Texture) final override
+	virtual FUnorderedAccessViewRHIRef RHICreateUnorderedAccessView(FTextureRHIParamRef Texture, uint32 MipLevel) final override
 	{ 
 		return new FRHIUnorderedAccessView(); 
 	}
diff --git a/Engine/Source/Runtime/OpenGLDrv/Private/OpenGLUAV.cpp b/Engine/Source/Runtime/OpenGLDrv/Private/OpenGLUAV.cpp
--- a/Engine/Source/Runtime/OpenGLDrv/Private/OpenGLUAV.cpp
+++ b/Engine/Source/Runtime/OpenGLDrv/Private/OpenGLUAV.cpp
@@ -56,7 +56,7 @@ FUnorderedAccessViewRHIRef FOpenGLDynamicRHI::RHICreateUnorderedAccessView(FStru
 	return new FOpenGLUnorderedAccessView();
 }
 
-FUnorderedAccessViewRHIRef FOpenGLDynamicRHI::RHICreateUnorderedAccessView(FTextureRHIParamRef TextureRHI)
+FUnorderedAccessViewRHIRef FOpenGLDynamicRHI::RHICreateUnorderedAccessView(FTextureRHIParamRef TextureRHI, uint32 MipLevel)
 {
 	FOpenGLTexture* Texture = ResourceCast(TextureRHI);
 	check(Texture->GetFlags() & TexCreate_UAV);
diff --git a/Engine/Source/Runtime/OpenGLDrv/Public/OpenGLDrv.h b/Engine/Source/Runtime/OpenGLDrv/Public/OpenGLDrv.h
--- a/Engine/Source/Runtime/OpenGLDrv/Public/OpenGLDrv.h
+++ b/Engine/Source/Runtime/OpenGLDrv/Public/OpenGLDrv.h
@@ -334,7 +334,7 @@ public:
 	virtual void* RHILockStructuredBuffer(FStructuredBufferRHIParamRef StructuredBuffer, uint32 Offset, uint32 SizeRHI, EResourceLockMode LockMode) final override;
 	virtual void RHIUnlockStructuredBuffer(FStructuredBufferRHIParamRef StructuredBuffer) final override;
 	virtual FUnorderedAccessViewRHIRef RHICreateUnorderedAccessView(FStructuredBufferRHIParamRef StructuredBuffer, bool bUseUAVCounter, bool bAppendBuffer) final override;
-	virtual FUnorderedAccessViewRHIRef RHICreateUnorderedAccessView(FTextureRHIParamRef Texture) final override;
+	virtual FUnorderedAccessViewRHIRef RHICreateUnorderedAccessView(FTextureRHIParamRef Texture, uint32 MipLevel) final override;
 	virtual FUnorderedAccessViewRHIRef RHICreateUnorderedAccessView(FVertexBufferRHIParamRef VertexBuffer, uint8 Format) final override;
 	virtual FShaderResourceViewRHIRef RHICreateShaderResourceView(FStructuredBufferRHIParamRef StructuredBuffer) final override;
 	virtual FShaderResourceViewRHIRef RHICreateShaderResourceView(FVertexBufferRHIParamRef VertexBuffer, uint32 Stride, uint8 Format) final override;
diff --git a/Engine/Source/Runtime/RHI/Public/DynamicRHI.h b/Engine/Source/Runtime/RHI/Public/DynamicRHI.h
--- a/Engine/Source/Runtime/RHI/Public/DynamicRHI.h
+++ b/Engine/Source/Runtime/RHI/Public/DynamicRHI.h
@@ -142,7 +142,7 @@ public:
 
 	/** Creates an unordered access view of the given texture. */
 	// FlushType: Wait RHI Thread
-	virtual FUnorderedAccessViewRHIRef RHICreateUnorderedAccessView(FTextureRHIParamRef Texture) = 0;
+	virtual FUnorderedAccessViewRHIRef RHICreateUnorderedAccessView(FTextureRHIParamRef Texture, uint32 MipLevel) = 0;
 
 	/** Creates an unordered access view of the given vertex buffer. */
 	// FlushType: Wait RHI Thread
diff --git a/Engine/Source/Runtime/RHI/Public/RHICommandList.h b/Engine/Source/Runtime/RHI/Public/RHICommandList.h
--- a/Engine/Source/Runtime/RHI/Public/RHICommandList.h
+++ b/Engine/Source/Runtime/RHI/Public/RHICommandList.h
@@ -573,20 +573,20 @@ struct FRHICommandSetRenderTargetsAndClear : public FRHICommand<FRHICommandSetRe
 };
 
 struct FRHICommandBindClearMRTValues : public FRHICommand<FRHICommandBindClearMRTValues>
-{	
+{
 	bool bClearColor;
 	bool bClearDepth;
 	bool bClearStencil;
 
 	FORCEINLINE_DEBUGGABLE FRHICommandBindClearMRTValues(
-		bool InbClearColor,		
-		bool InbClearDepth,		
+		bool InbClearColor,
+		bool InbClearDepth,
 		bool InbClearStencil
 		) 
 		: bClearColor(InbClearColor)
 		, bClearDepth(InbClearDepth)
 		, bClearStencil(InbClearStencil)
-	{		
+	{
 	}	
 
 	RHI_API void Execute(FRHICommandListBase& CmdList);
@@ -2114,14 +2114,14 @@ public:
 		return GDynamicRHI->RHICreateUnorderedAccessView(StructuredBuffer, bUseUAVCounter, bAppendBuffer);
 	}
 	
-	FORCEINLINE FUnorderedAccessViewRHIRef CreateUnorderedAccessView(FTextureRHIParamRef Texture)
+	FORCEINLINE FUnorderedAccessViewRHIRef CreateUnorderedAccessView(FTextureRHIParamRef Texture, uint32 MipLevel)
 	{
 		if(GRHIThread)
 		{
 			QUICK_SCOPE_CYCLE_COUNTER(STAT_RHIMETHOD_CreateUnorderedAccessView_WaitRHI);
 			ImmediateFlush(EImmediateFlushType::WaitForRHIThread);
 		}
-		return GDynamicRHI->RHICreateUnorderedAccessView(Texture);
+		return GDynamicRHI->RHICreateUnorderedAccessView(Texture, MipLevel);
 	}
 	
 	FORCEINLINE FUnorderedAccessViewRHIRef CreateUnorderedAccessView(FVertexBufferRHIParamRef VertexBuffer, uint8 Format)
@@ -2799,9 +2799,9 @@ FORCEINLINE FUnorderedAccessViewRHIRef RHICreateUnorderedAccessView(FStructuredB
 	return FRHICommandListExecutor::GetImmediateCommandList().CreateUnorderedAccessView(StructuredBuffer, bUseUAVCounter, bAppendBuffer);
 }
 
-FORCEINLINE FUnorderedAccessViewRHIRef RHICreateUnorderedAccessView(FTextureRHIParamRef Texture)
+FORCEINLINE FUnorderedAccessViewRHIRef RHICreateUnorderedAccessView(FTextureRHIParamRef Texture, uint32 MipLevel = 0)
 {
-	return FRHICommandListExecutor::GetImmediateCommandList().CreateUnorderedAccessView(Texture);
+	return FRHICommandListExecutor::GetImmediateCommandList().CreateUnorderedAccessView(Texture, MipLevel);
 }
 
 FORCEINLINE FUnorderedAccessViewRHIRef RHICreateUnorderedAccessView(FVertexBufferRHIParamRef VertexBuffer, uint8 Format)
diff --git a/Engine/Source/Runtime/Windows/D3D11RHI/Private/D3D11RHIPrivate.h b/Engine/Source/Runtime/Windows/D3D11RHI/Private/D3D11RHIPrivate.h
--- a/Engine/Source/Runtime/Windows/D3D11RHI/Private/D3D11RHIPrivate.h
+++ b/Engine/Source/Runtime/Windows/D3D11RHI/Private/D3D11RHIPrivate.h
@@ -329,7 +329,7 @@ public:
 	virtual void* RHILockStructuredBuffer(FStructuredBufferRHIParamRef StructuredBuffer, uint32 Offset, uint32 SizeRHI, EResourceLockMode LockMode) final override;
 	virtual void RHIUnlockStructuredBuffer(FStructuredBufferRHIParamRef StructuredBuffer) final override;
 	virtual FUnorderedAccessViewRHIRef RHICreateUnorderedAccessView(FStructuredBufferRHIParamRef StructuredBuffer, bool bUseUAVCounter, bool bAppendBuffer) final override;
-	virtual FUnorderedAccessViewRHIRef RHICreateUnorderedAccessView(FTextureRHIParamRef Texture) final override;
+	virtual FUnorderedAccessViewRHIRef RHICreateUnorderedAccessView(FTextureRHIParamRef Texture, uint32 MipLevel) final override;
 	virtual FUnorderedAccessViewRHIRef RHICreateUnorderedAccessView(FVertexBufferRHIParamRef VertexBuffer, uint8 Format) final override;
 	virtual FShaderResourceViewRHIRef RHICreateShaderResourceView(FStructuredBufferRHIParamRef StructuredBuffer) final override;
 	virtual FShaderResourceViewRHIRef RHICreateShaderResourceView(FVertexBufferRHIParamRef VertexBuffer, uint32 Stride, uint8 Format) final override;
diff --git a/Engine/Source/Runtime/Windows/D3D11RHI/Private/D3D11UAV.cpp b/Engine/Source/Runtime/Windows/D3D11RHI/Private/D3D11UAV.cpp
--- a/Engine/Source/Runtime/Windows/D3D11RHI/Private/D3D11UAV.cpp
+++ b/Engine/Source/Runtime/Windows/D3D11RHI/Private/D3D11UAV.cpp
@@ -53,7 +53,7 @@ FUnorderedAccessViewRHIRef FD3D11DynamicRHI::RHICreateUnorderedAccessView(FStruc
 	return new FD3D11UnorderedAccessView(UnorderedAccessView,StructuredBuffer);
 }
 
-FUnorderedAccessViewRHIRef FD3D11DynamicRHI::RHICreateUnorderedAccessView(FTextureRHIParamRef TextureRHI)
+FUnorderedAccessViewRHIRef FD3D11DynamicRHI::RHICreateUnorderedAccessView(FTextureRHIParamRef TextureRHI, uint32 MipLevel)
 {
 	FD3D11TextureBase* Texture = GetD3D11TextureFromRHITexture(TextureRHI);
 	
@@ -63,14 +63,30 @@ FUnorderedAccessViewRHIRef FD3D11DynamicRHI::RHICreateUnorderedAccessView(FTextu
 	{
 		FD3D11Texture3D* Texture3D = (FD3D11Texture3D*)Texture;
 		UAVDesc.ViewDimension = D3D11_UAV_DIMENSION_TEXTURE3D;
-		UAVDesc.Texture3D.MipSlice = 0;
+		UAVDesc.Texture3D.MipSlice = MipLevel;
 		UAVDesc.Texture3D.FirstWSlice = 0;
-		UAVDesc.Texture3D.WSize = Texture3D->GetSizeZ();
+		UAVDesc.Texture3D.WSize = Texture3D->GetSizeZ() >> MipLevel;
+	}
+	else if (TextureRHI->GetTexture2DArray() != NULL)
+	{
+		FD3D11Texture2DArray* Texture2DArray = (FD3D11Texture2DArray*)Texture;
+		UAVDesc.ViewDimension = D3D11_UAV_DIMENSION_TEXTURE2DARRAY;
+		UAVDesc.Texture2DArray.MipSlice = MipLevel;
+		UAVDesc.Texture2DArray.FirstArraySlice = 0;
+		UAVDesc.Texture2DArray.ArraySize = Texture2DArray->GetSizeZ();
+	}
+	else if (TextureRHI->GetTextureCube() != NULL)
+	{
+		FD3D11TextureCube* TextureCube = (FD3D11TextureCube*)Texture;
+		UAVDesc.ViewDimension = D3D11_UAV_DIMENSION_TEXTURE2DARRAY;
+		UAVDesc.Texture2DArray.MipSlice = MipLevel;
+		UAVDesc.Texture2DArray.FirstArraySlice = 0;
+		UAVDesc.Texture2DArray.ArraySize = 6;
 	}
 	else
 	{
 		UAVDesc.ViewDimension = D3D11_UAV_DIMENSION_TEXTURE2D;
-		UAVDesc.Texture2D.MipSlice = 0;
+		UAVDesc.Texture2D.MipSlice = MipLevel;
 	}
 	
 	UAVDesc.Format = FindShaderResourceDXGIFormat((DXGI_FORMAT)GPixelFormats[TextureRHI->GetFormat()].PlatformFormat, false);
