UE4.23 - Deprecate FUniformBufferRHIParamRef, FIndexBufferRHIParamRef, FVertexBufferRHIParamRef, FStructuredBufferRHIParamRef, FTextureRHIParamRef, FTexture2DRHIParamRef, FTexture2DArrayRHIParamRef, F

diff --git a/Engine/Source/Runtime/D3D12RHI/Private/D3D12VertexBuffer.cpp b/Engine/Source/Runtime/D3D12RHI/Private/D3D12VertexBuffer.cpp
--- a/Engine/Source/Runtime/D3D12RHI/Private/D3D12VertexBuffer.cpp
+++ b/Engine/Source/Runtime/D3D12RHI/Private/D3D12VertexBuffer.cpp
@@ -121,12 +121,12 @@ FVertexBufferRHIRef FD3D12DynamicRHI::RHICreateVertexBuffer(uint32 Size, uint32
 	return Buffer;
 }
 
-void* FD3D12DynamicRHI::RHILockVertexBuffer(FVertexBufferRHIParamRef VertexBufferRHI, uint32 Offset, uint32 Size, EResourceLockMode LockMode)
+void* FD3D12DynamicRHI::RHILockVertexBuffer(FRHIVertexBuffer* VertexBufferRHI, uint32 Offset, uint32 Size, EResourceLockMode LockMode)
 {
 	return LockBuffer(nullptr, FD3D12DynamicRHI::ResourceCast(VertexBufferRHI), Offset, Size, LockMode);
 }
 
-void FD3D12DynamicRHI::RHIUnlockVertexBuffer(FVertexBufferRHIParamRef VertexBufferRHI)
+void FD3D12DynamicRHI::RHIUnlockVertexBuffer(FRHIVertexBuffer* VertexBufferRHI)
 {
 	UnlockBuffer(nullptr, FD3D12DynamicRHI::ResourceCast(VertexBufferRHI));
 }
@@ -151,7 +151,7 @@ FVertexBufferRHIRef FD3D12DynamicRHI::CreateVertexBuffer_RenderThread(FRHIComman
 	return Buffer;
 }
 
-void* FD3D12DynamicRHI::LockVertexBuffer_RenderThread(class FRHICommandListImmediate& RHICmdList, FVertexBufferRHIParamRef VertexBufferRHI, uint32 Offset, uint32 SizeRHI, EResourceLockMode LockMode)
+void* FD3D12DynamicRHI::LockVertexBuffer_RenderThread(class FRHICommandListImmediate& RHICmdList, FRHIVertexBuffer* VertexBufferRHI, uint32 Offset, uint32 SizeRHI, EResourceLockMode LockMode)
 {
 	// Pull down the above RHI implementation so that we can flush only when absolutely necessary
 	QUICK_SCOPE_CYCLE_COUNTER(STAT_FDynamicRHI_LockVertexBuffer_RenderThread);
@@ -160,7 +160,7 @@ void* FD3D12DynamicRHI::LockVertexBuffer_RenderThread(class FRHICommandListImmed
 	return LockBuffer(&RHICmdList, FD3D12DynamicRHI::ResourceCast(VertexBufferRHI), Offset, SizeRHI, LockMode);
 }
 
-void FD3D12DynamicRHI::UnlockVertexBuffer_RenderThread(FRHICommandListImmediate& RHICmdList, FVertexBufferRHIParamRef VertexBufferRHI)
+void FD3D12DynamicRHI::UnlockVertexBuffer_RenderThread(FRHICommandListImmediate& RHICmdList, FRHIVertexBuffer* VertexBufferRHI)
 {
 	// Pull down the above RHI implementation so that we can flush only when absolutely necessary
 	QUICK_SCOPE_CYCLE_COUNTER(STAT_FDynamicRHI_UnlockVertexBuffer_RenderThread);
@@ -169,7 +169,7 @@ void FD3D12DynamicRHI::UnlockVertexBuffer_RenderThread(FRHICommandListImmediate&
 	UnlockBuffer(&RHICmdList, FD3D12DynamicRHI::ResourceCast(VertexBufferRHI));
 }
 
-void FD3D12DynamicRHI::RHICopyVertexBuffer(FVertexBufferRHIParamRef SourceBufferRHI, FVertexBufferRHIParamRef DestBufferRHI)
+void FD3D12DynamicRHI::RHICopyVertexBuffer(FRHIVertexBuffer* SourceBufferRHI, FRHIVertexBuffer* DestBufferRHI)
 {
 	FD3D12VertexBuffer*  SourceBuffer = FD3D12DynamicRHI::ResourceCast(SourceBufferRHI);
 	FD3D12VertexBuffer*  DestBuffer = FD3D12DynamicRHI::ResourceCast(DestBufferRHI);
@@ -203,7 +203,7 @@ void FD3D12DynamicRHI::RHICopyVertexBuffer(FVertexBufferRHIParamRef SourceBuffer
 	}
 }
 
-void FD3D12DynamicRHI::RHITransferVertexBufferUnderlyingResource(FVertexBufferRHIParamRef DestVertexBuffer, FVertexBufferRHIParamRef SrcVertexBuffer)
+void FD3D12DynamicRHI::RHITransferVertexBufferUnderlyingResource(FRHIVertexBuffer* DestVertexBuffer, FRHIVertexBuffer* SrcVertexBuffer)
 {
 	check(DestVertexBuffer);
 	FD3D12VertexBuffer* Dest = ResourceCast(DestVertexBuffer);
@@ -219,7 +219,7 @@ void FD3D12DynamicRHI::RHITransferVertexBufferUnderlyingResource(FVertexBufferRH
 }
 
 #if D3D12_RHI_RAYTRACING
-void FD3D12CommandContext::RHICopyBufferRegion(FVertexBufferRHIParamRef DestBufferRHI, uint64 DstOffset, FVertexBufferRHIParamRef SourceBufferRHI, uint64 SrcOffset, uint64 NumBytes)
+void FD3D12CommandContext::RHICopyBufferRegion(FRHIVertexBuffer* DestBufferRHI, uint64 DstOffset, FRHIVertexBuffer* SourceBufferRHI, uint64 SrcOffset, uint64 NumBytes)
 {
 	FD3D12VertexBuffer*  SourceBuffer = RetrieveObject<FD3D12VertexBuffer>(SourceBufferRHI);
 	FD3D12VertexBuffer*  DestBuffer = RetrieveObject<FD3D12VertexBuffer>(DestBufferRHI);
diff --git a/Engine/Source/Runtime/D3D12RHI/Private/D3D12View.cpp b/Engine/Source/Runtime/D3D12RHI/Private/D3D12View.cpp
--- a/Engine/Source/Runtime/D3D12RHI/Private/D3D12View.cpp
+++ b/Engine/Source/Runtime/D3D12RHI/Private/D3D12View.cpp
@@ -80,7 +80,7 @@ FD3D12ShaderResourceView* CreateSRV(TextureType* Texture, D3D12_SHADER_RESOURCE_
 	});
 }
 
-FShaderResourceViewRHIRef FD3D12DynamicRHI::RHICreateShaderResourceView(FTextureRHIParamRef Texture, const FRHITextureSRVCreateInfo& CreateInfo)
+FShaderResourceViewRHIRef FD3D12DynamicRHI::RHICreateShaderResourceView(FRHITexture* Texture, const FRHITextureSRVCreateInfo& CreateInfo)
 {
 	D3D12_SHADER_RESOURCE_VIEW_DESC SRVDesc = {};
 	SRVDesc.Shader4ComponentMapping = D3D12_DEFAULT_SHADER_4_COMPONENT_MAPPING;
@@ -153,7 +153,7 @@ FShaderResourceViewRHIRef FD3D12DynamicRHI::RHICreateShaderResourceView(FTexture
 	return CreateSRV(BaseTexture, SRVDesc);
 }
 
-FShaderResourceViewRHIRef FD3D12DynamicRHI::RHICreateShaderResourceView(FStructuredBufferRHIParamRef StructuredBufferRHI)
+FShaderResourceViewRHIRef FD3D12DynamicRHI::RHICreateShaderResourceView(FRHIStructuredBuffer* StructuredBufferRHI)
 {
 	struct FD3D12InitializeStructuredBufferSRVRHICommand final : public FRHICommand<FD3D12InitializeStructuredBufferSRVRHICommand>
 	{
@@ -237,7 +237,7 @@ FShaderResourceViewRHIRef FD3D12DynamicRHI::RHICreateShaderResourceView(FStructu
 	});
 }
 
-FShaderResourceViewRHIRef FD3D12DynamicRHI::RHICreateShaderResourceView(FVertexBufferRHIParamRef VertexBufferRHI, uint32 Stride, uint8 Format)
+FShaderResourceViewRHIRef FD3D12DynamicRHI::RHICreateShaderResourceView(FRHIVertexBuffer* VertexBufferRHI, uint32 Stride, uint8 Format)
 {
 	struct FD3D12InitializeVertexBufferSRVRHICommand final : public FRHICommand<FD3D12InitializeVertexBufferSRVRHICommand>
 	{
@@ -297,7 +297,7 @@ FShaderResourceViewRHIRef FD3D12DynamicRHI::RHICreateShaderResourceView(FVertexB
 	});
 }
 
-FShaderResourceViewRHIRef FD3D12DynamicRHI::RHICreateShaderResourceView(FIndexBufferRHIParamRef BufferRHI)
+FShaderResourceViewRHIRef FD3D12DynamicRHI::RHICreateShaderResourceView(FRHIIndexBuffer* BufferRHI)
 {
 	if (!BufferRHI)
 	{
@@ -318,7 +318,7 @@ FShaderResourceViewRHIRef FD3D12DynamicRHI::RHICreateShaderResourceView(FIndexBu
 	});
 }
 
-void FD3D12DynamicRHI::RHIUpdateShaderResourceView(FRHIShaderResourceView* SRV, FVertexBufferRHIParamRef VertexBuffer, uint32 Stride, uint8 Format)
+void FD3D12DynamicRHI::RHIUpdateShaderResourceView(FRHIShaderResourceView* SRV, FRHIVertexBuffer* VertexBuffer, uint32 Stride, uint8 Format)
 {
 	check(SRV);
 	if (VertexBuffer)
@@ -343,7 +343,7 @@ void FD3D12DynamicRHI::RHIUpdateShaderResourceView(FRHIShaderResourceView* SRV,
 	}
 }
 
-void FD3D12DynamicRHI::RHIUpdateShaderResourceView(FRHIShaderResourceView* SRV, FIndexBufferRHIParamRef IndexBuffer)
+void FD3D12DynamicRHI::RHIUpdateShaderResourceView(FRHIShaderResourceView* SRV, FRHIIndexBuffer* IndexBuffer)
 {
 	check(SRV);
 	if (IndexBuffer)
@@ -368,22 +368,22 @@ void FD3D12DynamicRHI::RHIUpdateShaderResourceView(FRHIShaderResourceView* SRV,
 	}
 }
 
-FShaderResourceViewRHIRef FD3D12DynamicRHI::RHICreateShaderResourceView_RenderThread(FRHICommandListImmediate& RHICmdList, FTextureRHIParamRef Texture, const FRHITextureSRVCreateInfo& CreateInfo)
+FShaderResourceViewRHIRef FD3D12DynamicRHI::RHICreateShaderResourceView_RenderThread(FRHICommandListImmediate& RHICmdList, FRHITexture* Texture, const FRHITextureSRVCreateInfo& CreateInfo)
 {
 	return RHICreateShaderResourceView(Texture, CreateInfo);
 }
 
-FShaderResourceViewRHIRef FD3D12DynamicRHI::RHICreateShaderResourceView_RenderThread(FRHICommandListImmediate& RHICmdList, FVertexBufferRHIParamRef VertexBufferRHI, uint32 Stride, uint8 Format)
+FShaderResourceViewRHIRef FD3D12DynamicRHI::RHICreateShaderResourceView_RenderThread(FRHICommandListImmediate& RHICmdList, FRHIVertexBuffer* VertexBufferRHI, uint32 Stride, uint8 Format)
 {
 	return RHICreateShaderResourceView(VertexBufferRHI, Stride, Format);
 }
 
-FShaderResourceViewRHIRef FD3D12DynamicRHI::CreateShaderResourceView_RenderThread(FRHICommandListImmediate& RHICmdList, FVertexBufferRHIParamRef VertexBufferRHI, uint32 Stride, uint8 Format)
+FShaderResourceViewRHIRef FD3D12DynamicRHI::CreateShaderResourceView_RenderThread(FRHICommandListImmediate& RHICmdList, FRHIVertexBuffer* VertexBufferRHI, uint32 Stride, uint8 Format)
 {
 	return RHICreateShaderResourceView_RenderThread(RHICmdList, VertexBufferRHI, Stride, Format);
 }
 
-FShaderResourceViewRHIRef FD3D12DynamicRHI::RHICreateShaderResourceView_RenderThread(FRHICommandListImmediate& RHICmdList, FStructuredBufferRHIParamRef StructuredBufferRHI)
+FShaderResourceViewRHIRef FD3D12DynamicRHI::RHICreateShaderResourceView_RenderThread(FRHICommandListImmediate& RHICmdList, FRHIStructuredBuffer* StructuredBufferRHI)
 {
 	return RHICreateShaderResourceView(StructuredBufferRHI);
 }
diff --git a/Engine/Source/Runtime/D3D12RHI/Private/D3D12Viewport.cpp b/Engine/Source/Runtime/D3D12RHI/Private/D3D12Viewport.cpp
--- a/Engine/Source/Runtime/D3D12RHI/Private/D3D12Viewport.cpp
+++ b/Engine/Source/Runtime/D3D12RHI/Private/D3D12Viewport.cpp
@@ -806,7 +806,7 @@ void FD3D12DynamicRHI::RHITick(float DeltaTime)
  *	Viewport functions.
  *=============================================================================*/
 
-void FD3D12CommandContextBase::RHIBeginDrawingViewport(FRHIViewport* ViewportRHI, FTextureRHIParamRef RenderTargetRHI)
+void FD3D12CommandContextBase::RHIBeginDrawingViewport(FRHIViewport* ViewportRHI, FRHITexture* RenderTargetRHI)
 {
 	FD3D12Viewport* Viewport = FD3D12DynamicRHI::ResourceCast(ViewportRHI);
 
diff --git a/Engine/Source/Runtime/EmptyRHI/Private/EmptyCommands.cpp b/Engine/Source/Runtime/EmptyRHI/Private/EmptyCommands.cpp
--- a/Engine/Source/Runtime/EmptyRHI/Private/EmptyCommands.cpp
+++ b/Engine/Source/Runtime/EmptyRHI/Private/EmptyCommands.cpp
@@ -6,13 +6,13 @@
 
 #include "EmptyRHIPrivate.h"
 
-void FEmptyDynamicRHI::RHISetStreamSource(uint32 StreamIndex,FVertexBufferRHIParamRef VertexBufferRHI, uint32 Offset)
+void FEmptyDynamicRHI::RHISetStreamSource(uint32 StreamIndex, FRHIVertexBuffer* VertexBufferRHI, uint32 Offset)
 {
 	FEmptyVertexBuffer* VertexBuffer = ResourceCast(VertexBufferRHI);
 
 }
 
-void FEmptyDynamicRHI::RHISetStreamOutTargets(uint32 NumTargets, const FVertexBufferRHIParamRef* VertexBuffers, const uint32* Offsets)
+void FEmptyDynamicRHI::RHISetStreamOutTargets(uint32 NumTargets, const FRHIVertexBuffer** VertexBuffers, const uint32* Offsets)
 {
 
 }
@@ -34,7 +34,7 @@ void FEmptyDynamicRHI::RHIDispatchComputeShader(uint32 ThreadGroupCountX, uint32
 
 }
 
-void FEmptyDynamicRHI::RHIDispatchIndirectComputeShader(FVertexBufferRHIParamRef ArgumentBufferRHI, uint32 ArgumentOffset) 
+void FEmptyDynamicRHI::RHIDispatchIndirectComputeShader(FRHIVertexBuffer* ArgumentBufferRHI, uint32 ArgumentOffset)
 { 
 	FEmptyVertexBuffer* ArgumentBuffer = ResourceCast(ArgumentBufferRHI);
 
@@ -75,32 +75,32 @@ void FEmptyDynamicRHI::RHISetUAVParameter(FRHIComputeShader* ComputeShaderRHI,ui
 }
 
 
-void FEmptyDynamicRHI::RHISetShaderTexture(FRHIVertexShader* VertexShaderRHI, uint32 TextureIndex, FTextureRHIParamRef NewTextureRHI)
+void FEmptyDynamicRHI::RHISetShaderTexture(FRHIVertexShader* VertexShaderRHI, uint32 TextureIndex, FRHITexture* NewTextureRHI)
 {
 
 }
 
-void FEmptyDynamicRHI::RHISetShaderTexture(FRHIHullShader* HullShader, uint32 TextureIndex, FTextureRHIParamRef NewTextureRHI)
+void FEmptyDynamicRHI::RHISetShaderTexture(FRHIHullShader* HullShader, uint32 TextureIndex, FRHITexture* NewTextureRHI)
 {
 
 }
 
-void FEmptyDynamicRHI::RHISetShaderTexture(FRHIDomainShader* DomainShader, uint32 TextureIndex, FTextureRHIParamRef NewTextureRHI)
+void FEmptyDynamicRHI::RHISetShaderTexture(FRHIDomainShader* DomainShader, uint32 TextureIndex, FRHITexture* NewTextureRHI)
 {
 
 }
 
-void FEmptyDynamicRHI::RHISetShaderTexture(FRHIGeometryShader* GeometryShader, uint32 TextureIndex, FTextureRHIParamRef NewTextureRHI)
+void FEmptyDynamicRHI::RHISetShaderTexture(FRHIGeometryShader* GeometryShader, uint32 TextureIndex, FRHITexture* NewTextureRHI)
 {
 
 }
 
-void FEmptyDynamicRHI::RHISetShaderTexture(FRHIPixelShader* PixelShader, uint32 TextureIndex, FTextureRHIParamRef NewTextureRHI)
+void FEmptyDynamicRHI::RHISetShaderTexture(FRHIPixelShader* PixelShader, uint32 TextureIndex, FRHITexture* NewTextureRHI)
 {
 
 }
 
-void FEmptyDynamicRHI::RHISetShaderTexture(FRHIComputeShader* ComputeShader, uint32 TextureIndex, FTextureRHIParamRef NewTextureRHI)
+void FEmptyDynamicRHI::RHISetShaderTexture(FRHIComputeShader* ComputeShader, uint32 TextureIndex, FRHITexture* NewTextureRHI)
 {
 
 }
@@ -205,34 +205,34 @@ void FEmptyDynamicRHI::RHISetShaderParameter(FRHIComputeShader* ComputeShaderRHI
 
 }
 
-void FEmptyDynamicRHI::RHISetShaderUniformBuffer(FRHIVertexShader* VertexShaderRHI, uint32 BufferIndex, FUniformBufferRHIParamRef BufferRHI)
+void FEmptyDynamicRHI::RHISetShaderUniformBuffer(FRHIVertexShader* VertexShaderRHI, uint32 BufferIndex, FRHIUniformBuffer* BufferRHI)
 {
 	FEmptyUniformBuffer* Buffer = ResourceCast(BufferRHI);
 	FEmptyVertexShader* VertexShader = ResourceCast(VertexShaderRHI);
 
 }
 
-void FEmptyDynamicRHI::RHISetShaderUniformBuffer(FRHIHullShader* HullShader, uint32 BufferIndex, FUniformBufferRHIParamRef BufferRHI)
+void FEmptyDynamicRHI::RHISetShaderUniformBuffer(FRHIHullShader* HullShader, uint32 BufferIndex, FRHIUniformBuffer* BufferRHI)
 {
 	FEmptyUniformBuffer* Buffer = ResourceCast(BufferRHI);
 }
 
-void FEmptyDynamicRHI::RHISetShaderUniformBuffer(FRHIDomainShader* DomainShader, uint32 BufferIndex, FUniformBufferRHIParamRef BufferRHI)
+void FEmptyDynamicRHI::RHISetShaderUniformBuffer(FRHIDomainShader* DomainShader, uint32 BufferIndex, FRHIUniformBuffer* BufferRHI)
 {
 	FEmptyUniformBuffer* Buffer = ResourceCast(BufferRHI);
 }
 
-void FEmptyDynamicRHI::RHISetShaderUniformBuffer(FRHIGeometryShader* GeometryShader, uint32 BufferIndex, FUniformBufferRHIParamRef BufferRHI)
+void FEmptyDynamicRHI::RHISetShaderUniformBuffer(FRHIGeometryShader* GeometryShader, uint32 BufferIndex, FRHIUniformBuffer* BufferRHI)
 {
 	FEmptyUniformBuffer* Buffer = ResourceCast(BufferRHI);
 }
 
-void FEmptyDynamicRHI::RHISetShaderUniformBuffer(FRHIPixelShader* PixelShader, uint32 BufferIndex, FUniformBufferRHIParamRef BufferRHI)
+void FEmptyDynamicRHI::RHISetShaderUniformBuffer(FRHIPixelShader* PixelShader, uint32 BufferIndex, FRHIUniformBuffer* BufferRHI)
 {
 	FEmptyUniformBuffer* Buffer = ResourceCast(BufferRHI);
 }
 
-void FEmptyDynamicRHI::RHISetShaderUniformBuffer(FRHIComputeShader* ComputeShader, uint32 BufferIndex, FUniformBufferRHIParamRef BufferRHI)
+void FEmptyDynamicRHI::RHISetShaderUniformBuffer(FRHIComputeShader* ComputeShader, uint32 BufferIndex, FRHIUniformBuffer* BufferRHI)
 {
 	FEmptyUniformBuffer* Buffer = ResourceCast(BufferRHI);
 }
@@ -252,7 +252,7 @@ void FEmptyDynamicRHI::RHISetBlendState(FRHIBlendState* NewStateRHI, const FLine
 
 
 void FEmptyDynamicRHI::RHISetRenderTargets(uint32 NumSimultaneousRenderTargets, const FRHIRenderTargetView* NewRenderTargets, 
-	FTextureRHIParamRef NewDepthStencilTargetRHI, uint32 NumUAVs, FRHIUnorderedAccessView* const* UAVs)
+	FRHITexture* NewDepthStencilTargetRHI, uint32 NumUAVs, FRHIUnorderedAccessView* const* UAVs)
 {
 
 }
@@ -281,28 +281,28 @@ void FEmptyDynamicRHI::RHIDrawPrimitive(uint32 PrimitiveType, uint32 BaseVertexI
 {
 }
 
-void FEmptyDynamicRHI::RHIDrawPrimitiveIndirect(uint32 PrimitiveType, FVertexBufferRHIParamRef ArgumentBufferRHI, uint32 ArgumentOffset)
+void FEmptyDynamicRHI::RHIDrawPrimitiveIndirect(uint32 PrimitiveType, FRHIVertexBuffer* ArgumentBufferRHI, uint32 ArgumentOffset)
 {
 	FEmptyVertexBuffer* ArgumentBuffer = ResourceCast(ArgumentBufferRHI);
 
 }
 
 
-void FEmptyDynamicRHI::RHIDrawIndexedPrimitive(FIndexBufferRHIParamRef IndexBufferRHI, uint32 PrimitiveType, int32 BaseVertexIndex, uint32 FirstInstance,
+void FEmptyDynamicRHI::RHIDrawIndexedPrimitive(FRHIIndexBuffer* IndexBufferRHI, uint32 PrimitiveType, int32 BaseVertexIndex, uint32 FirstInstance,
 	uint32 NumVertices, uint32 StartIndex, uint32 NumPrimitives, uint32 NumInstances)
 {
 	FEmptyIndexBuffer* IndexBuffer = ResourceCast(IndexBufferRHI);
 
 }
 
-void FEmptyDynamicRHI::RHIDrawIndexedIndirect(FIndexBufferRHIParamRef IndexBufferRHI, uint32 PrimitiveType, FStructuredBufferRHIParamRef ArgumentsBufferRHI, int32 DrawArgumentsIndex, uint32 NumInstances)
+void FEmptyDynamicRHI::RHIDrawIndexedIndirect(FRHIIndexBuffer* IndexBufferRHI, uint32 PrimitiveType, FRHIStructuredBuffer* ArgumentsBufferRHI, int32 DrawArgumentsIndex, uint32 NumInstances)
 {
 	FEmptyIndexBuffer* IndexBuffer = ResourceCast(IndexBufferRHI);
 	FEmptyStructuredBuffer* ArgumentsBuffer = ResourceCast(ArgumentsBufferRHI);
 
 }
 
-void FEmptyDynamicRHI::RHIDrawIndexedPrimitiveIndirect(uint32 PrimitiveType,FIndexBufferRHIParamRef IndexBufferRHI,FVertexBufferRHIParamRef ArgumentBufferRHI,uint32 ArgumentOffset)
+void FEmptyDynamicRHI::RHIDrawIndexedPrimitiveIndirect(uint32 PrimitiveType, FRHIIndexBuffer* IndexBufferRHI, FRHIVertexBuffer* ArgumentBufferRHI,uint32 ArgumentOffset)
 {
 	FEmptyIndexBuffer* IndexBuffer = ResourceCast(IndexBufferRHI);
 	FEmptyVertexBuffer* ArgumentBuffer = ResourceCast(ArgumentBufferRHI);
diff --git a/Engine/Source/Runtime/EmptyRHI/Private/EmptyIndexBuffer.cpp b/Engine/Source/Runtime/EmptyRHI/Private/EmptyIndexBuffer.cpp
--- a/Engine/Source/Runtime/EmptyRHI/Private/EmptyIndexBuffer.cpp
+++ b/Engine/Source/Runtime/EmptyRHI/Private/EmptyIndexBuffer.cpp
@@ -52,21 +52,21 @@ FIndexBufferRHIRef FEmptyDynamicRHI::RHICreateIndexBuffer(uint32 Stride, uint32
 	return IndexBuffer;
 }
 
-void* FEmptyDynamicRHI::RHILockIndexBuffer(FIndexBufferRHIParamRef IndexBufferRHI, uint32 Offset, uint32 Size, EResourceLockMode LockMode)
+void* FEmptyDynamicRHI::RHILockIndexBuffer(FRHIIndexBuffer* IndexBufferRHI, uint32 Offset, uint32 Size, EResourceLockMode LockMode)
 {
 	FEmptyIndexBuffer* IndexBuffer = ResourceCast(IndexBufferRHI);
 
 	return (uint8*)IndexBuffer->Lock(LockMode, Size) + Offset;
 }
 
-void FEmptyDynamicRHI::RHIUnlockIndexBuffer(FIndexBufferRHIParamRef IndexBufferRHI)
+void FEmptyDynamicRHI::RHIUnlockIndexBuffer(FRHIIndexBuffer* IndexBufferRHI)
 {
 	FEmptyIndexBuffer* IndexBuffer = ResourceCast(IndexBufferRHI);
 
 	IndexBuffer->Unlock();
 }
 
-void FEmptyDynamicRHI::RHITransferIndexBufferUnderlyingResource(FIndexBufferRHIParamRef DestIndexBuffer, FIndexBufferRHIParamRef SrcIndexBuffer)
+void FEmptyDynamicRHI::RHITransferIndexBufferUnderlyingResource(FRHIIndexBuffer* DestIndexBuffer, FRHIIndexBuffer* SrcIndexBuffer)
 {
 
 }
diff --git a/Engine/Source/Runtime/EmptyRHI/Private/EmptyRenderTarget.cpp b/Engine/Source/Runtime/EmptyRHI/Private/EmptyRenderTarget.cpp
--- a/Engine/Source/Runtime/EmptyRHI/Private/EmptyRenderTarget.cpp
+++ b/Engine/Source/Runtime/EmptyRHI/Private/EmptyRenderTarget.cpp
@@ -8,32 +8,32 @@
 #include "ScreenRendering.h"
 
 
-void FEmptyDynamicRHI::RHICopyToResolveTarget(FTextureRHIParamRef SourceTextureRHI, FTextureRHIParamRef DestTextureRHI, const FResolveParams& ResolveParams)
+void FEmptyDynamicRHI::RHICopyToResolveTarget(FRHITexture* SourceTextureRHI, FRHITexture* DestTextureRHI, const FResolveParams& ResolveParams)
 {
 
 }
 
-void FEmptyDynamicRHI::RHIReadSurfaceData(FTextureRHIParamRef TextureRHI, FIntRect Rect, TArray<FColor>& OutData, FReadSurfaceDataFlags InFlags)
+void FEmptyDynamicRHI::RHIReadSurfaceData(FRHITexture* TextureRHI, FIntRect Rect, TArray<FColor>& OutData, FReadSurfaceDataFlags InFlags)
 {
 
 }
 
-void FEmptyDynamicRHI::RHIMapStagingSurface(FTextureRHIParamRef TextureRHI,void*& OutData,int32& OutWidth,int32& OutHeight)
+void FEmptyDynamicRHI::RHIMapStagingSurface(FRHITexture* TextureRHI,void*& OutData,int32& OutWidth,int32& OutHeight)
 {
 
 }
 
-void FEmptyDynamicRHI::RHIUnmapStagingSurface(FTextureRHIParamRef TextureRHI)
+void FEmptyDynamicRHI::RHIUnmapStagingSurface(FRHITexture* TextureRHI)
 {
 
 }
 
-void FEmptyDynamicRHI::RHIReadSurfaceFloatData(FTextureRHIParamRef TextureRHI, FIntRect Rect, TArray<FFloat16Color>& OutData, ECubeFace CubeFace,int32 ArrayIndex,int32 MipIndex)
+void FEmptyDynamicRHI::RHIReadSurfaceFloatData(FRHITexture* TextureRHI, FIntRect Rect, TArray<FFloat16Color>& OutData, ECubeFace CubeFace,int32 ArrayIndex,int32 MipIndex)
 {
 
 }
 
-void FEmptyDynamicRHI::RHIRead3DSurfaceFloatData(FTextureRHIParamRef TextureRHI,FIntRect InRect,FIntPoint ZMinMax,TArray<FFloat16Color>& OutData)
+void FEmptyDynamicRHI::RHIRead3DSurfaceFloatData(FRHITexture* TextureRHI,FIntRect InRect,FIntPoint ZMinMax,TArray<FFloat16Color>& OutData)
 {
 
 }
diff --git a/Engine/Source/Runtime/EmptyRHI/Private/EmptyStructuredBuffer.cpp b/Engine/Source/Runtime/EmptyRHI/Private/EmptyStructuredBuffer.cpp
--- a/Engine/Source/Runtime/EmptyRHI/Private/EmptyStructuredBuffer.cpp
+++ b/Engine/Source/Runtime/EmptyRHI/Private/EmptyStructuredBuffer.cpp
@@ -30,14 +30,14 @@ FStructuredBufferRHIRef FEmptyDynamicRHI::RHICreateStructuredBuffer(uint32 Strid
 	return nullptr;//new FEmptyStructuredBuffer(Stride, Size, ResourceArray, InUsage);
 }
 
-void* FEmptyDynamicRHI::RHILockStructuredBuffer(FStructuredBufferRHIParamRef StructuredBufferRHI,uint32 Offset,uint32 Size,EResourceLockMode LockMode)
+void* FEmptyDynamicRHI::RHILockStructuredBuffer(FRHIStructuredBuffer* StructuredBufferRHI,uint32 Offset,uint32 Size,EResourceLockMode LockMode)
 {
 	FEmptyStructuredBuffer* StructuredBuffer = ResourceCast(StructuredBufferRHI);
 
 	return nullptr;
 }
 
-void FEmptyDynamicRHI::RHIUnlockStructuredBuffer(FStructuredBufferRHIParamRef StructuredBufferRHI)
+void FEmptyDynamicRHI::RHIUnlockStructuredBuffer(FRHIStructuredBuffer* StructuredBufferRHI)
 {
 
 }
diff --git a/Engine/Source/Runtime/EmptyRHI/Private/EmptyTexture.cpp b/Engine/Source/Runtime/EmptyRHI/Private/EmptyTexture.cpp
--- a/Engine/Source/Runtime/EmptyRHI/Private/EmptyTexture.cpp
+++ b/Engine/Source/Runtime/EmptyRHI/Private/EmptyTexture.cpp
@@ -98,7 +98,7 @@ bool FEmptyDynamicRHI::RHIGetTextureMemoryVisualizeData( FColor* /*TextureData*/
 	return false;
 }
 
-uint32 FEmptyDynamicRHI::RHIComputeMemorySize(FTextureRHIParamRef TextureRHI)
+uint32 FEmptyDynamicRHI::RHIComputeMemorySize(FRHITexture* TextureRHI)
 {
 	if(!TextureRHI)
 	{
@@ -123,7 +123,7 @@ FTexture2DRHIRef FEmptyDynamicRHI::RHIAsyncCreateTexture2D(uint32 SizeX,uint32 S
 	return FTexture2DRHIRef();
 }
 
-void FEmptyDynamicRHI::RHICopySharedMips(FTexture2DRHIParamRef DestTexture2D,FTexture2DRHIParamRef SrcTexture2D)
+void FEmptyDynamicRHI::RHICopySharedMips(FRHITexture2D* DestTexture2D, FRHITexture2D* SrcTexture2D)
 {
 
 }
@@ -138,64 +138,64 @@ FTexture3DRHIRef FEmptyDynamicRHI::RHICreateTexture3D(uint32 SizeX, uint32 SizeY
 	return new FEmptyTexture3D((EPixelFormat)Format, SizeX, SizeY, SizeZ, NumMips, Flags, CreateInfo.BulkData);
 }
 
-void FEmptyDynamicRHI::RHIGetResourceInfo(FTextureRHIParamRef Ref, FRHIResourceInfo& OutInfo)
+void FEmptyDynamicRHI::RHIGetResourceInfo(FRHITexture* Ref, FRHIResourceInfo& OutInfo)
 {
 }
 
-void FEmptyDynamicRHI::RHIGenerateMips(FTextureRHIParamRef SourceSurfaceRHI)
+void FEmptyDynamicRHI::RHIGenerateMips(FRHITexture* SourceSurfaceRHI)
 {
 
 }
 
-FTexture2DRHIRef FEmptyDynamicRHI::RHIAsyncReallocateTexture2D(FTexture2DRHIParamRef OldTextureRHI, int32 NewMipCount, int32 NewSizeX, int32 NewSizeY, FThreadSafeCounter* RequestStatus)
+FTexture2DRHIRef FEmptyDynamicRHI::RHIAsyncReallocateTexture2D(FRHITexture2D* OldTextureRHI, int32 NewMipCount, int32 NewSizeX, int32 NewSizeY, FThreadSafeCounter* RequestStatus)
 {
 	FEmptyTexture2D* OldTexture = ResourceCast(OldTextureRHI);
 
 	return NULL;
 }
 
-ETextureReallocationStatus FEmptyDynamicRHI::RHIFinalizeAsyncReallocateTexture2D( FTexture2DRHIParamRef Texture2D, bool bBlockUntilCompleted )
+ETextureReallocationStatus FEmptyDynamicRHI::RHIFinalizeAsyncReallocateTexture2D(FRHITexture2D* Texture2D, bool bBlockUntilCompleted )
 {
 	return TexRealloc_Failed;
 }
 
-ETextureReallocationStatus FEmptyDynamicRHI::RHICancelAsyncReallocateTexture2D( FTexture2DRHIParamRef Texture2D, bool bBlockUntilCompleted )
+ETextureReallocationStatus FEmptyDynamicRHI::RHICancelAsyncReallocateTexture2D(FRHITexture2D* Texture2D, bool bBlockUntilCompleted )
 {
 	return TexRealloc_Failed;
 }
 
 
-void* FEmptyDynamicRHI::RHILockTexture2D(FTexture2DRHIParamRef TextureRHI,uint32 MipIndex,EResourceLockMode LockMode,uint32& DestStride,bool bLockWithinMiptail)
+void* FEmptyDynamicRHI::RHILockTexture2D(FRHITexture2D* TextureRHI,uint32 MipIndex,EResourceLockMode LockMode,uint32& DestStride,bool bLockWithinMiptail)
 {
 	FEmptyTexture2D* Texture = ResourceCast(TextureRHI);
 	return Texture->Surface.Lock(MipIndex, 0, LockMode, DestStride);
 }
 
-void FEmptyDynamicRHI::RHIUnlockTexture2D(FTexture2DRHIParamRef TextureRHI,uint32 MipIndex,bool bLockWithinMiptail)
+void FEmptyDynamicRHI::RHIUnlockTexture2D(FRHITexture2D* TextureRHI,uint32 MipIndex,bool bLockWithinMiptail)
 {
 	FEmptyTexture2D* Texture = ResourceCast(TextureRHI);
 	Texture->Surface.Unlock(MipIndex, 0);
 }
 
-void* FEmptyDynamicRHI::RHILockTexture2DArray(FTexture2DArrayRHIParamRef TextureRHI,uint32 TextureIndex,uint32 MipIndex,EResourceLockMode LockMode,uint32& DestStride,bool bLockWithinMiptail)
+void* FEmptyDynamicRHI::RHILockTexture2DArray(FRHITexture2DArray* TextureRHI,uint32 TextureIndex,uint32 MipIndex,EResourceLockMode LockMode,uint32& DestStride,bool bLockWithinMiptail)
 {
 	FEmptyTexture2DArray* Texture = ResourceCast(TextureRHI);
 	return Texture->Surface.Lock(MipIndex, TextureIndex, LockMode, DestStride);
 }
 
-void FEmptyDynamicRHI::RHIUnlockTexture2DArray(FTexture2DArrayRHIParamRef TextureRHI,uint32 TextureIndex,uint32 MipIndex,bool bLockWithinMiptail)
+void FEmptyDynamicRHI::RHIUnlockTexture2DArray(FRHITexture2DArray* TextureRHI,uint32 TextureIndex,uint32 MipIndex,bool bLockWithinMiptail)
 {
 	FEmptyTexture2DArray* Texture = ResourceCast(TextureRHI);
 	Texture->Surface.Unlock(MipIndex, TextureIndex);
 }
 
-void FEmptyDynamicRHI::RHIUpdateTexture2D(FTexture2DRHIParamRef TextureRHI, uint32 MipIndex, const struct FUpdateTextureRegion2D& UpdateRegion, uint32 SourcePitch, const uint8* SourceData)
+void FEmptyDynamicRHI::RHIUpdateTexture2D(FRHITexture2D* TextureRHI, uint32 MipIndex, const struct FUpdateTextureRegion2D& UpdateRegion, uint32 SourcePitch, const uint8* SourceData)
 {	
 	FEmptyTexture2D* Texture = ResourceCast(TextureRHI);
 
 }
 
-void FEmptyDynamicRHI::RHIUpdateTexture3D(FTexture3DRHIParamRef TextureRHI,uint32 MipIndex,const FUpdateTextureRegion3D& UpdateRegion,uint32 SourceRowPitch,uint32 SourceDepthPitch,const uint8* SourceData)
+void FEmptyDynamicRHI::RHIUpdateTexture3D(FRHITexture3D* TextureRHI,uint32 MipIndex,const FUpdateTextureRegion3D& UpdateRegion,uint32 SourceRowPitch,uint32 SourceDepthPitch,const uint8* SourceData)
 {	
 	FEmptyTexture3D* Texture = ResourceCast(TextureRHI);
 
@@ -215,30 +215,30 @@ FTextureCubeRHIRef FEmptyDynamicRHI::RHICreateTextureCubeArray(uint32 Size, uint
 	return new FEmptyTextureCube((EPixelFormat)Format, Size, true, ArraySize, NumMips, Flags, CreateInfo.BulkData);
 }
 
-void* FEmptyDynamicRHI::RHILockTextureCubeFace(FTextureCubeRHIParamRef TextureCubeRHI,uint32 FaceIndex,uint32 ArrayIndex,uint32 MipIndex,EResourceLockMode LockMode,uint32& DestStride,bool bLockWithinMiptail)
+void* FEmptyDynamicRHI::RHILockTextureCubeFace(FRHITextureCube* TextureCubeRHI,uint32 FaceIndex,uint32 ArrayIndex,uint32 MipIndex,EResourceLockMode LockMode,uint32& DestStride,bool bLockWithinMiptail)
 {
 	FEmptyTextureCube* TextureCube = ResourceCast(TextureCubeRHI);
 	uint32 EmptyFace = GetEmptyCubeFace((ECubeFace)FaceIndex);
 	return TextureCube->Surface.Lock(MipIndex, FaceIndex + 6 * ArrayIndex, LockMode, DestStride);
 }
 
-void FEmptyDynamicRHI::RHIUnlockTextureCubeFace(FTextureCubeRHIParamRef TextureCubeRHI,uint32 FaceIndex,uint32 ArrayIndex,uint32 MipIndex,bool bLockWithinMiptail)
+void FEmptyDynamicRHI::RHIUnlockTextureCubeFace(FRHITextureCube* TextureCubeRHI,uint32 FaceIndex,uint32 ArrayIndex,uint32 MipIndex,bool bLockWithinMiptail)
 {
 	FEmptyTextureCube* TextureCube = ResourceCast(TextureCubeRHI);
 	uint32 EmptyFace = GetEmptyCubeFace((ECubeFace)FaceIndex);
 	TextureCube->Surface.Unlock(MipIndex, FaceIndex + ArrayIndex * 6);
 }
 
-void FEmptyDynamicRHI::RHIBindDebugLabelName(FTextureRHIParamRef TextureRHI, const TCHAR* Name)
+void FEmptyDynamicRHI::RHIBindDebugLabelName(FRHITexture* TextureRHI, const TCHAR* Name)
 {
 }
 
-void FEmptyDynamicRHI::RHIVirtualTextureSetFirstMipInMemory(FTexture2DRHIParamRef TextureRHI, uint32 FirstMip)
+void FEmptyDynamicRHI::RHIVirtualTextureSetFirstMipInMemory(FRHITexture2D* TextureRHI, uint32 FirstMip)
 {
 
 }
 
-void FEmptyDynamicRHI::RHIVirtualTextureSetFirstMipVisible(FTexture2DRHIParamRef TextureRHI, uint32 FirstMip)
+void FEmptyDynamicRHI::RHIVirtualTextureSetFirstMipVisible(FRHITexture2D* TextureRHI, uint32 FirstMip)
 {
 }
 
@@ -262,6 +262,6 @@ FTextureReferenceRHIRef FEmptyDynamicRHI::RHICreateTextureReference(FLastRenderT
 	return nullptr;
 }
 
-void FEmptyDynamicRHI::RHIUpdateTextureReference(FRHITextureReference* TextureRef, FTextureRHIParamRef NewTexture)
+void FEmptyDynamicRHI::RHIUpdateTextureReference(FRHITextureReference* TextureRef, FRHITexture* NewTexture)
 {
 }
diff --git a/Engine/Source/Runtime/EmptyRHI/Private/EmptyUAV.cpp b/Engine/Source/Runtime/EmptyRHI/Private/EmptyUAV.cpp
--- a/Engine/Source/Runtime/EmptyRHI/Private/EmptyUAV.cpp
+++ b/Engine/Source/Runtime/EmptyRHI/Private/EmptyUAV.cpp
@@ -11,7 +11,7 @@ FEmptyShaderResourceView::~FEmptyShaderResourceView()
 
 
 
-FUnorderedAccessViewRHIRef FEmptyDynamicRHI::RHICreateUnorderedAccessView(FStructuredBufferRHIParamRef StructuredBufferRHI, bool bUseUAVCounter, bool bAppendBuffer)
+FUnorderedAccessViewRHIRef FEmptyDynamicRHI::RHICreateUnorderedAccessView(FRHIStructuredBuffer* StructuredBufferRHI, bool bUseUAVCounter, bool bAppendBuffer)
 {
 	FEmptyStructuredBuffer* StructuredBuffer = ResourceCast(StructuredBufferRHI);
 
@@ -22,7 +22,7 @@ FUnorderedAccessViewRHIRef FEmptyDynamicRHI::RHICreateUnorderedAccessView(FStruc
 	return UAV;
 }
 
-FUnorderedAccessViewRHIRef FEmptyDynamicRHI::RHICreateUnorderedAccessView(FTextureRHIParamRef TextureRHI, uint32 MipLevel)
+FUnorderedAccessViewRHIRef FEmptyDynamicRHI::RHICreateUnorderedAccessView(FRHITexture* TextureRHI, uint32 MipLevel)
 {
 	FEmptySurface& Surface = GetEmptySurfaceFromRHITexture(TextureRHI);
 
@@ -33,7 +33,7 @@ FUnorderedAccessViewRHIRef FEmptyDynamicRHI::RHICreateUnorderedAccessView(FTextu
 	return UAV;
 }
 
-FUnorderedAccessViewRHIRef FEmptyDynamicRHI::RHICreateUnorderedAccessView(FVertexBufferRHIParamRef VertexBufferRHI, uint8 Format)
+FUnorderedAccessViewRHIRef FEmptyDynamicRHI::RHICreateUnorderedAccessView(FRHIVertexBuffer* VertexBufferRHI, uint8 Format)
 {
 	FEmptyVertexBuffer* VertexBuffer = ResourceCast(VertexBufferRHI);
 
@@ -44,7 +44,7 @@ FUnorderedAccessViewRHIRef FEmptyDynamicRHI::RHICreateUnorderedAccessView(FVerte
 	return UAV;
 }
 
-FUnorderedAccessViewRHIRef FEmptyDynamicRHI::RHICreateUnorderedAccessView(FIndexBufferRHIParamRef IndexBufferRHI, uint8 Format)
+FUnorderedAccessViewRHIRef FEmptyDynamicRHI::RHICreateUnorderedAccessView(FRHIIndexBuffer* IndexBufferRHI, uint8 Format)
 {
 	FEmptyIndexBuffer* IndexBuffer = ResourceCast(IndexBufferRHI);
 
@@ -55,7 +55,7 @@ FUnorderedAccessViewRHIRef FEmptyDynamicRHI::RHICreateUnorderedAccessView(FIndex
 	return UAV;
 }
 
-FShaderResourceViewRHIRef FEmptyDynamicRHI::RHICreateShaderResourceView(FStructuredBufferRHIParamRef StructuredBufferRHI)
+FShaderResourceViewRHIRef FEmptyDynamicRHI::RHICreateShaderResourceView(FRHIStructuredBuffer* StructuredBufferRHI)
 {
 	FEmptyStructuredBuffer* StructuredBuffer = ResourceCast(StructuredBufferRHI);
 
@@ -63,7 +63,7 @@ FShaderResourceViewRHIRef FEmptyDynamicRHI::RHICreateShaderResourceView(FStructu
 	return SRV;
 }
 
-FShaderResourceViewRHIRef FEmptyDynamicRHI::RHICreateShaderResourceView(FVertexBufferRHIParamRef VertexBufferRHI, uint32 Stride, uint8 Format)
+FShaderResourceViewRHIRef FEmptyDynamicRHI::RHICreateShaderResourceView(FRHIVertexBuffer* VertexBufferRHI, uint32 Stride, uint8 Format)
 {
 	FEmptyVertexBuffer* VertexBuffer = ResourceCast(VertexBufferRHI);
 
@@ -72,41 +72,41 @@ FShaderResourceViewRHIRef FEmptyDynamicRHI::RHICreateShaderResourceView(FVertexB
 	return SRV;
 }
 
-FShaderResourceViewRHIRef FEmptyDynamicRHI::RHICreateShaderResourceView(FIndexBufferRHIParamRef BufferRHI)
+FShaderResourceViewRHIRef FEmptyDynamicRHI::RHICreateShaderResourceView(FRHIIndexBuffer* BufferRHI)
 {
 	// there should be no need to create an object
 	return FShaderResourceViewRHIRef();
 }
 
-FShaderResourceViewRHIRef FEmptyDynamicRHI::RHICreateShaderResourceView(FTexture2DRHIParamRef Texture2DRHI, uint8 MipLevel)
+FShaderResourceViewRHIRef FEmptyDynamicRHI::RHICreateShaderResourceView(FRHITexture2D* Texture2DRHI, uint8 MipLevel)
 {
 	FEmptyShaderResourceView* SRV = new FEmptyShaderResourceView;
 	SRV->SourceTexture = (FRHITexture*)Texture2DRHI;
 	return SRV;
 }
 
-FShaderResourceViewRHIRef FEmptyDynamicRHI::RHICreateShaderResourceView(FTexture2DRHIParamRef Texture2DRHI, uint8 MipLevel, uint8 NumMipLevels, uint8 Format)
+FShaderResourceViewRHIRef FEmptyDynamicRHI::RHICreateShaderResourceView(FRHITexture2D* Texture2DRHI, uint8 MipLevel, uint8 NumMipLevels, uint8 Format)
 {
 	FEmptyShaderResourceView* SRV = new FEmptyShaderResourceView;
 	SRV->SourceTexture = (FRHITexture*)Texture2DRHI;
 	return SRV;
 }
 
-FShaderResourceViewRHIRef FEmptyDynamicRHI::RHICreateShaderResourceView(FTexture3DRHIParamRef Texture3DRHI, uint8 MipLevel)
+FShaderResourceViewRHIRef FEmptyDynamicRHI::RHICreateShaderResourceView(FRHITexture3D* Texture3DRHI, uint8 MipLevel)
 {
 	FEmptyShaderResourceView* SRV = new FEmptyShaderResourceView;
 	SRV->SourceTexture = (FRHITexture*)Texture3DRHI;
 	return SRV;
 }
 
-FShaderResourceViewRHIRef FEmptyDynamicRHI::RHICreateShaderResourceView(FTexture2DArrayRHIParamRef Texture2DArrayRHI, uint8 MipLevel)
+FShaderResourceViewRHIRef FEmptyDynamicRHI::RHICreateShaderResourceView(FRHITexture2DArray* Texture2DArrayRHI, uint8 MipLevel)
 {
 	FEmptyShaderResourceView* SRV = new FEmptyShaderResourceView;
 	SRV->SourceTexture = (FRHITexture*)Texture2DArrayRHI;
 	return SRV;
 }
 
-FShaderResourceViewRHIRef FEmptyDynamicRHI::RHICreateShaderResourceView(FTextureCubeRHIParamRef TextureCubeRHI, uint8 MipLevel)
+FShaderResourceViewRHIRef FEmptyDynamicRHI::RHICreateShaderResourceView(FRHITextureCube* TextureCubeRHI, uint8 MipLevel)
 {
 	FEmptyShaderResourceView* SRV = new FEmptyShaderResourceView;
 	SRV->SourceTexture = (FRHITexture*)TextureCubeRHI;
diff --git a/Engine/Source/Runtime/EmptyRHI/Private/EmptyVertexBuffer.cpp b/Engine/Source/Runtime/EmptyRHI/Private/EmptyVertexBuffer.cpp
--- a/Engine/Source/Runtime/EmptyRHI/Private/EmptyVertexBuffer.cpp
+++ b/Engine/Source/Runtime/EmptyRHI/Private/EmptyVertexBuffer.cpp
@@ -51,7 +51,7 @@ FVertexBufferRHIRef FEmptyDynamicRHI::RHICreateVertexBuffer(uint32 Size, uint32
 	return VertexBuffer;
 }
 
-void* FEmptyDynamicRHI::RHILockVertexBuffer(FVertexBufferRHIParamRef VertexBufferRHI, uint32 Offset, uint32 Size, EResourceLockMode LockMode)
+void* FEmptyDynamicRHI::RHILockVertexBuffer(FRHIVertexBuffer* VertexBufferRHI, uint32 Offset, uint32 Size, EResourceLockMode LockMode)
 {
 	FEmptyVertexBuffer* VertexBuffer = ResourceCast(VertexBufferRHI);
 
@@ -59,19 +59,19 @@ void* FEmptyDynamicRHI::RHILockVertexBuffer(FVertexBufferRHIParamRef VertexBuffe
 	return (uint8*)VertexBuffer->Lock(LockMode, Size) + Offset;
 }
 
-void FEmptyDynamicRHI::RHIUnlockVertexBuffer(FVertexBufferRHIParamRef VertexBufferRHI)
+void FEmptyDynamicRHI::RHIUnlockVertexBuffer(FRHIVertexBuffer* VertexBufferRHI)
 {
 	FEmptyVertexBuffer* VertexBuffer = ResourceCast(VertexBufferRHI);
 
 	VertexBuffer->Unlock();
 }
 
-void FEmptyDynamicRHI::RHICopyVertexBuffer(FVertexBufferRHIParamRef SourceBufferRHI,FVertexBufferRHIParamRef DestBufferRHI)
+void FEmptyDynamicRHI::RHICopyVertexBuffer(FRHIVertexBuffer* SourceBufferRHI, FRHIVertexBuffer* DestBufferRHI)
 {
 
 }
 
-void FEmptyDynamicRHI::RHITransferVertexBufferUnderlyingResource(FVertexBufferRHIParamRef DestVertexBuffer, FVertexBufferRHIParamRef SrcVertexBuffer)
+void FEmptyDynamicRHI::RHITransferVertexBufferUnderlyingResource(FRHIVertexBuffer* DestVertexBuffer, FRHIVertexBuffer* SrcVertexBuffer)
 {
 
 }
\ No newline at end of file
diff --git a/Engine/Source/Runtime/EmptyRHI/Private/EmptyViewport.cpp b/Engine/Source/Runtime/EmptyRHI/Private/EmptyViewport.cpp
--- a/Engine/Source/Runtime/EmptyRHI/Private/EmptyViewport.cpp
+++ b/Engine/Source/Runtime/EmptyRHI/Private/EmptyViewport.cpp
@@ -43,7 +43,7 @@ void FEmptyDynamicRHI::RHITick( float DeltaTime )
  *	Viewport functions.
  *=============================================================================*/
 
-void FEmptyDynamicRHI::RHIBeginDrawingViewport(FRHIViewport* ViewportRHI, FTextureRHIParamRef RenderTargetRHI)
+void FEmptyDynamicRHI::RHIBeginDrawingViewport(FRHIViewport* ViewportRHI, FRHITexture* RenderTargetRHI)
 {
 	FEmptyViewport* Viewport = ResourceCast(ViewportRHI);
 
diff --git a/Engine/Source/Runtime/EmptyRHI/Public/EmptyRHI.h b/Engine/Source/Runtime/EmptyRHI/Public/EmptyRHI.h
--- a/Engine/Source/Runtime/EmptyRHI/Public/EmptyRHI.h
+++ b/Engine/Source/Runtime/EmptyRHI/Public/EmptyRHI.h
@@ -51,63 +51,63 @@ public:
 	virtual FBoundShaderStateRHIRef RHICreateBoundShaderState(FRHIVertexDeclaration* VertexDeclaration, FRHIVertexShader* VertexShader, FRHIHullShader* HullShader, FRHIDomainShaderFRHIDomainShader*, FRHIPixelShader* PixelShader, FRHIGeometryShader* GeometryShader) final override;
 	virtual FUniformBufferRHIRef RHICreateUniformBuffer(const void* Contents, const FRHIUniformBufferLayout& Layout, EUniformBufferUsage Usage, EUniformBufferValidation Validation) final override;
 	virtual FIndexBufferRHIRef RHICreateIndexBuffer(uint32 Stride, uint32 Size, uint32 InUsage, FRHIResourceCreateInfo& CreateInfo) final override;
-	virtual void* RHILockIndexBuffer(FIndexBufferRHIParamRef IndexBuffer, uint32 Offset, uint32 Size, EResourceLockMode LockMode) final override;
-	virtual void RHIUnlockIndexBuffer(FIndexBufferRHIParamRef IndexBuffer) final override;
-	virtual void RHITransferIndexBufferUnderlyingResource(FIndexBufferRHIParamRef DestIndexBuffer, FIndexBufferRHIParamRef SrcIndexBuffer) final override;
+	virtual void* RHILockIndexBuffer(FRHIIndexBuffer* IndexBuffer, uint32 Offset, uint32 Size, EResourceLockMode LockMode) final override;
+	virtual void RHIUnlockIndexBuffer(FRHIIndexBuffer* IndexBuffer) final override;
+	virtual void RHITransferIndexBufferUnderlyingResource(FRHIIndexBuffer* DestIndexBuffer, FRHIIndexBuffer* SrcIndexBuffer) final override;
 	virtual FVertexBufferRHIRef RHICreateVertexBuffer(uint32 Size, uint32 InUsage, FRHIResourceCreateInfo& CreateInfo) final override;
-	virtual void* RHILockVertexBuffer(FVertexBufferRHIParamRef VertexBuffer, uint32 Offset, uint32 SizeRHI, EResourceLockMode LockMode) final override;
-	virtual void RHIUnlockVertexBuffer(FVertexBufferRHIParamRef VertexBuffer) final override;
-	virtual void RHICopyVertexBuffer(FVertexBufferRHIParamRef SourceBuffer,FVertexBufferRHIParamRef DestBuffer) final override;
-	virtual void RHITransferVertexBufferUnderlyingResource(FVertexBufferRHIParamRef DestVertexBuffer, FVertexBufferRHIParamRef SrcVertexBuffer) final override;
+	virtual void* RHILockVertexBuffer(FRHIVertexBuffer* VertexBuffer, uint32 Offset, uint32 SizeRHI, EResourceLockMode LockMode) final override;
+	virtual void RHIUnlockVertexBuffer(FRHIVertexBuffer* VertexBuffer) final override;
+	virtual void RHICopyVertexBuffer(FRHIVertexBuffer* SourceBuffer, FRHIVertexBuffer* DestBuffer) final override;
+	virtual void RHITransferVertexBufferUnderlyingResource(FRHIVertexBuffer* DestVertexBuffer, FRHIVertexBuffer* SrcVertexBuffer) final override;
 	virtual FStructuredBufferRHIRef RHICreateStructuredBuffer(uint32 Stride, uint32 Size, uint32 InUsage, FRHIResourceCreateInfo& CreateInfo) final override;
-	virtual void* RHILockStructuredBuffer(FStructuredBufferRHIParamRef StructuredBuffer, uint32 Offset, uint32 SizeRHI, EResourceLockMode LockMode) final override;
-	virtual void RHIUnlockStructuredBuffer(FStructuredBufferRHIParamRef StructuredBuffer) final override;
-	virtual FUnorderedAccessViewRHIRef RHICreateUnorderedAccessView(FStructuredBufferRHIParamRef StructuredBuffer, bool bUseUAVCounter, bool bAppendBuffer) final override;
-	virtual FUnorderedAccessViewRHIRef RHICreateUnorderedAccessView(FTextureRHIParamRef Texture, uint32 MipLevel) final override;
-	virtual FUnorderedAccessViewRHIRef RHICreateUnorderedAccessView(FVertexBufferRHIParamRef VertexBuffer, uint8 Format) final override;
-	virtual FUnorderedAccessViewRHIRef RHICreateUnorderedAccessView(FIndexBufferRHIParamRef IndexBuffer, uint8 Format) final override;
-	virtual FShaderResourceViewRHIRef RHICreateShaderResourceView(FStructuredBufferRHIParamRef StructuredBuffer) final override;
-	virtual FShaderResourceViewRHIRef RHICreateShaderResourceView(FVertexBufferRHIParamRef VertexBuffer, uint32 Stride, uint8 Format) final override;
-	virtual FShaderResourceViewRHIRef RHICreateShaderResourceView(FIndexBufferRHIParamRef Buffer) final override;
+	virtual void* RHILockStructuredBuffer(FRHIStructuredBuffer* StructuredBuffer, uint32 Offset, uint32 SizeRHI, EResourceLockMode LockMode) final override;
+	virtual void RHIUnlockStructuredBuffer(FRHIStructuredBuffer* StructuredBuffer) final override;
+	virtual FUnorderedAccessViewRHIRef RHICreateUnorderedAccessView(FRHIStructuredBuffer* StructuredBuffer, bool bUseUAVCounter, bool bAppendBuffer) final override;
+	virtual FUnorderedAccessViewRHIRef RHICreateUnorderedAccessView(FRHITexture* Texture, uint32 MipLevel) final override;
+	virtual FUnorderedAccessViewRHIRef RHICreateUnorderedAccessView(FRHIVertexBuffer* VertexBuffer, uint8 Format) final override;
+	virtual FUnorderedAccessViewRHIRef RHICreateUnorderedAccessView(FRHIIndexBuffer* IndexBuffer, uint8 Format) final override;
+	virtual FShaderResourceViewRHIRef RHICreateShaderResourceView(FRHIStructuredBuffer* StructuredBuffer) final override;
+	virtual FShaderResourceViewRHIRef RHICreateShaderResourceView(FRHIVertexBuffer* VertexBuffer, uint32 Stride, uint8 Format) final override;
+	virtual FShaderResourceViewRHIRef RHICreateShaderResourceView(FRHIIndexBuffer* Buffer) final override;
 	virtual uint64 RHICalcTexture2DPlatformSize(uint32 SizeX, uint32 SizeY, uint8 Format, uint32 NumMips, uint32 NumSamples, uint32 Flags, uint32& OutAlign) final override;
 	virtual uint64 RHICalcTexture3DPlatformSize(uint32 SizeX, uint32 SizeY, uint32 SizeZ, uint8 Format, uint32 NumMips, uint32 Flags, uint32& OutAlign) final override;
 	virtual uint64 RHICalcTextureCubePlatformSize(uint32 Size, uint8 Format, uint32 NumMips, uint32 Flags, uint32& OutAlign) final override;
 	virtual void RHIGetTextureMemoryStats(FTextureMemoryStats& OutStats) final override;
 	virtual bool RHIGetTextureMemoryVisualizeData(FColor* TextureData,int32 SizeX,int32 SizeY,int32 Pitch,int32 PixelSize) final override;
 	virtual FTextureReferenceRHIRef RHICreateTextureReference(FLastRenderTimeContainer* LastRenderTime) final override;
-	virtual void RHIUpdateTextureReference(FRHITextureReference* TextureRef, FTextureRHIParamRef NewTexture) final override;
+	virtual void RHIUpdateTextureReference(FRHITextureReference* TextureRef, FRHITexture* NewTexture) final override;
 	virtual FTexture2DRHIRef RHICreateTexture2D(uint32 SizeX, uint32 SizeY, uint8 Format, uint32 NumMips, uint32 NumSamples, uint32 Flags, FRHIResourceCreateInfo& CreateInfo) final override;
 	virtual FTexture2DRHIRef RHIAsyncCreateTexture2D(uint32 SizeX, uint32 SizeY, uint8 Format, uint32 NumMips, uint32 Flags, void** InitialMipData, uint32 NumInitialMips) final override;
-	virtual void RHICopySharedMips(FTexture2DRHIParamRef DestTexture2D, FTexture2DRHIParamRef SrcTexture2D) final override;
+	virtual void RHICopySharedMips(FRHITexture2D* DestTexture2D, FRHITexture2D* SrcTexture2D) final override;
 	virtual FTexture2DArrayRHIRef RHICreateTexture2DArray(uint32 SizeX, uint32 SizeY, uint32 SizeZ, uint8 Format, uint32 NumMips, uint32 NumSamples, uint32 Flags, FRHIResourceCreateInfo& CreateInfo) final override;
 	virtual FTexture3DRHIRef RHICreateTexture3D(uint32 SizeX, uint32 SizeY, uint32 SizeZ, uint8 Format, uint32 NumMips, uint32 Flags, FRHIResourceCreateInfo& CreateInfo) final override;
-	virtual void RHIGetResourceInfo(FTextureRHIParamRef Ref, FRHIResourceInfo& OutInfo) final override;
-	virtual FShaderResourceViewRHIRef RHICreateShaderResourceView(FTexture2DRHIParamRef Texture2DRHI, uint8 MipLevel) final override;
-	virtual FShaderResourceViewRHIRef RHICreateShaderResourceView(FTexture2DRHIParamRef Texture2DRHI, uint8 MipLevel, uint8 NumMipLevels, uint8 Format) final override;
-	virtual FShaderResourceViewRHIRef RHICreateShaderResourceView(FTexture3DRHIParamRef Texture3DRHI, uint8 MipLevel) final override;
-	virtual FShaderResourceViewRHIRef RHICreateShaderResourceView(FTexture2DArrayRHIParamRef Texture2DArrayRHI, uint8 MipLevel) final override;
-	virtual FShaderResourceViewRHIRef RHICreateShaderResourceView(FTextureCubeRHIParamRef TextureCubeRHI, uint8 MipLevel) final override;
-	virtual void RHIGenerateMips(FTextureRHIParamRef Texture) final override;
-	virtual uint32 RHIComputeMemorySize(FTextureRHIParamRef TextureRHI) final override;
-	virtual FTexture2DRHIRef RHIAsyncReallocateTexture2D(FTexture2DRHIParamRef Texture2D, int32 NewMipCount, int32 NewSizeX, int32 NewSizeY, FThreadSafeCounter* RequestStatus) final override;
-	virtual ETextureReallocationStatus RHIFinalizeAsyncReallocateTexture2D(FTexture2DRHIParamRef Texture2D, bool bBlockUntilCompleted) final override;
-	virtual ETextureReallocationStatus RHICancelAsyncReallocateTexture2D(FTexture2DRHIParamRef Texture2D, bool bBlockUntilCompleted) final override;
-	virtual void* RHILockTexture2D(FTexture2DRHIParamRef Texture, uint32 MipIndex, EResourceLockMode LockMode, uint32& DestStride, bool bLockWithinMiptail) final override;
-	virtual void RHIUnlockTexture2D(FTexture2DRHIParamRef Texture, uint32 MipIndex, bool bLockWithinMiptail) final override;
-	virtual void* RHILockTexture2DArray(FTexture2DArrayRHIParamRef Texture, uint32 TextureIndex, uint32 MipIndex, EResourceLockMode LockMode, uint32& DestStride, bool bLockWithinMiptail) final override;
-	virtual void RHIUnlockTexture2DArray(FTexture2DArrayRHIParamRef Texture, uint32 TextureIndex, uint32 MipIndex, bool bLockWithinMiptail) final override;
-	virtual void RHIUpdateTexture2D(FTexture2DRHIParamRef Texture, uint32 MipIndex, const struct FUpdateTextureRegion2D& UpdateRegion, uint32 SourcePitch, const uint8* SourceData) final override;
-	virtual void RHIUpdateTexture3D(FTexture3DRHIParamRef Texture, uint32 MipIndex, const struct FUpdateTextureRegion3D& UpdateRegion, uint32 SourceRowPitch, uint32 SourceDepthPitch, const uint8* SourceData) final override;
+	virtual void RHIGetResourceInfo(FRHITexture* Ref, FRHIResourceInfo& OutInfo) final override;
+	virtual FShaderResourceViewRHIRef RHICreateShaderResourceView(FRHITexture2D* Texture2DRHI, uint8 MipLevel) final override;
+	virtual FShaderResourceViewRHIRef RHICreateShaderResourceView(FRHITexture2D* Texture2DRHI, uint8 MipLevel, uint8 NumMipLevels, uint8 Format) final override;
+	virtual FShaderResourceViewRHIRef RHICreateShaderResourceView(FRHITexture3D* Texture3DRHI, uint8 MipLevel) final override;
+	virtual FShaderResourceViewRHIRef RHICreateShaderResourceView(FRHITexture2DArray* Texture2DArrayRHI, uint8 MipLevel) final override;
+	virtual FShaderResourceViewRHIRef RHICreateShaderResourceView(FRHITextureCube* TextureCubeRHI, uint8 MipLevel) final override;
+	virtual void RHIGenerateMips(FRHITexture* Texture) final override;
+	virtual uint32 RHIComputeMemorySize(FRHITexture* TextureRHI) final override;
+	virtual FTexture2DRHIRef RHIAsyncReallocateTexture2D(FRHITexture2D* Texture2D, int32 NewMipCount, int32 NewSizeX, int32 NewSizeY, FThreadSafeCounter* RequestStatus) final override;
+	virtual ETextureReallocationStatus RHIFinalizeAsyncReallocateTexture2D(FRHITexture2D* Texture2D, bool bBlockUntilCompleted) final override;
+	virtual ETextureReallocationStatus RHICancelAsyncReallocateTexture2D(FRHITexture2D* Texture2D, bool bBlockUntilCompleted) final override;
+	virtual void* RHILockTexture2D(FRHITexture2D* Texture, uint32 MipIndex, EResourceLockMode LockMode, uint32& DestStride, bool bLockWithinMiptail) final override;
+	virtual void RHIUnlockTexture2D(FRHITexture2D* Texture, uint32 MipIndex, bool bLockWithinMiptail) final override;
+	virtual void* RHILockTexture2DArray(FRHITexture2DArray* Texture, uint32 TextureIndex, uint32 MipIndex, EResourceLockMode LockMode, uint32& DestStride, bool bLockWithinMiptail) final override;
+	virtual void RHIUnlockTexture2DArray(FRHITexture2DArray* Texture, uint32 TextureIndex, uint32 MipIndex, bool bLockWithinMiptail) final override;
+	virtual void RHIUpdateTexture2D(FRHITexture2D* Texture, uint32 MipIndex, const struct FUpdateTextureRegion2D& UpdateRegion, uint32 SourcePitch, const uint8* SourceData) final override;
+	virtual void RHIUpdateTexture3D(FRHITexture3D* Texture, uint32 MipIndex, const struct FUpdateTextureRegion3D& UpdateRegion, uint32 SourceRowPitch, uint32 SourceDepthPitch, const uint8* SourceData) final override;
 	virtual FTextureCubeRHIRef RHICreateTextureCube(uint32 Size, uint8 Format, uint32 NumMips, uint32 Flags, FRHIResourceCreateInfo& CreateInfo) final override;
 	virtual FTextureCubeRHIRef RHICreateTextureCubeArray(uint32 Size, uint32 ArraySize, uint8 Format, uint32 NumMips, uint32 Flags, FRHIResourceCreateInfo& CreateInfo) final override;
-	virtual void* RHILockTextureCubeFace(FTextureCubeRHIParamRef Texture, uint32 FaceIndex, uint32 ArrayIndex, uint32 MipIndex, EResourceLockMode LockMode, uint32& DestStride, bool bLockWithinMiptail) final override;
-	virtual void RHIUnlockTextureCubeFace(FTextureCubeRHIParamRef Texture, uint32 FaceIndex, uint32 ArrayIndex, uint32 MipIndex, bool bLockWithinMiptail) final override;
-	virtual void RHIBindDebugLabelName(FTextureRHIParamRef Texture, const TCHAR* Name) final override;
-	virtual void RHIReadSurfaceData(FTextureRHIParamRef Texture,FIntRect Rect,TArray<FColor>& OutData,FReadSurfaceDataFlags InFlags) final override;
-	virtual void RHIMapStagingSurface(FTextureRHIParamRef Texture,void*& OutData,int32& OutWidth,int32& OutHeight) final override;
-	virtual void RHIUnmapStagingSurface(FTextureRHIParamRef Texture) final override;
-	virtual void RHIReadSurfaceFloatData(FTextureRHIParamRef Texture,FIntRect Rect,TArray<FFloat16Color>& OutData,ECubeFace CubeFace,int32 ArrayIndex,int32 MipIndex) final override;
-	virtual void RHIRead3DSurfaceFloatData(FTextureRHIParamRef Texture,FIntRect Rect,FIntPoint ZMinMax,TArray<FFloat16Color>& OutData) final override;
+	virtual void* RHILockTextureCubeFace(FRHITextureCube* Texture, uint32 FaceIndex, uint32 ArrayIndex, uint32 MipIndex, EResourceLockMode LockMode, uint32& DestStride, bool bLockWithinMiptail) final override;
+	virtual void RHIUnlockTextureCubeFace(FRHITextureCube* Texture, uint32 FaceIndex, uint32 ArrayIndex, uint32 MipIndex, bool bLockWithinMiptail) final override;
+	virtual void RHIBindDebugLabelName(FRHITexture* Texture, const TCHAR* Name) final override;
+	virtual void RHIReadSurfaceData(FRHITexture* Texture,FIntRect Rect,TArray<FColor>& OutData,FReadSurfaceDataFlags InFlags) final override;
+	virtual void RHIMapStagingSurface(FRHITexture* Texture,void*& OutData,int32& OutWidth,int32& OutHeight) final override;
+	virtual void RHIUnmapStagingSurface(FRHITexture* Texture) final override;
+	virtual void RHIReadSurfaceFloatData(FRHITexture* Texture,FIntRect Rect,TArray<FFloat16Color>& OutData,ECubeFace CubeFace,int32 ArrayIndex,int32 MipIndex) final override;
+	virtual void RHIRead3DSurfaceFloatData(FRHITexture* Texture,FIntRect Rect,FIntPoint ZMinMax,TArray<FFloat16Color>& OutData) final override;
 	virtual FRenderQueryRHIRef RHICreateRenderQuery(ERenderQueryType QueryType) final override;
 	virtual bool RHIGetRenderQueryResult(FRHIRenderQuery* RenderQuery, uint64& OutResult, bool bWait) final override;
 	virtual FTexture2DRHIRef RHIGetViewportBackBuffer(FRHIViewport* Viewport) final override;
@@ -119,12 +119,12 @@ public:
 	virtual FViewportRHIRef RHICreateViewport(void* WindowHandle, uint32 SizeX, uint32 SizeY, bool bIsFullscreen, EPixelFormat PreferredPixelFormat) final override;
 	virtual void RHIResizeViewport(FRHIViewport* Viewport, uint32 SizeX, uint32 SizeY, bool bIsFullscreen) final override;
 	virtual void RHITick(float DeltaTime) final override;
-	virtual void RHISetStreamOutTargets(uint32 NumTargets,const FVertexBufferRHIParamRef* VertexBuffers,const uint32* Offsets) final override;
+	virtual void RHISetStreamOutTargets(uint32 NumTargets,const FRHIVertexBuffer** VertexBuffers,const uint32* Offsets) final override;
 	virtual void RHIBlockUntilGPUIdle() final override;
 	virtual bool RHIGetAvailableResolutions(FScreenResolutionArray& Resolutions, bool bIgnoreRefreshRate) final override;
 	virtual void RHIGetSupportedResolution(uint32& Width, uint32& Height) final override;
-	virtual void RHIVirtualTextureSetFirstMipInMemory(FTexture2DRHIParamRef Texture, uint32 FirstMip) final override;
-	virtual void RHIVirtualTextureSetFirstMipVisible(FTexture2DRHIParamRef Texture, uint32 FirstMip) final override;
+	virtual void RHIVirtualTextureSetFirstMipInMemory(FRHITexture2D* Texture, uint32 FirstMip) final override;
+	virtual void RHIVirtualTextureSetFirstMipVisible(FRHITexture2D* Texture, uint32 FirstMip) final override;
 	virtual void RHIExecuteCommandList(FRHICommandList* CmdList) final override;
 	virtual void* RHIGetNativeDevice() final override;
 	virtual class IRHICommandContext* RHIGetDefaultContext() final override;
@@ -134,33 +134,33 @@ public:
 
 	virtual void RHISetComputeShader(FRHIComputeShader* ComputeShader) final override;
 	virtual void RHIDispatchComputeShader(uint32 ThreadGroupCountX, uint32 ThreadGroupCountY, uint32 ThreadGroupCountZ) final override;
-	virtual void RHIDispatchIndirectComputeShader(FVertexBufferRHIParamRef ArgumentBuffer, uint32 ArgumentOffset) final override;
+	virtual void RHIDispatchIndirectComputeShader(FRHIVertexBuffer* ArgumentBuffer, uint32 ArgumentOffset) final override;
 	virtual void RHIAutomaticCacheFlushAfterComputeShader(bool bEnable) final override;
 	virtual void RHIFlushComputeShaderCache() final override;
 	virtual void RHISetMultipleViewports(uint32 Count, const FViewportBounds* Data) final override;
 	virtual void RHIClearTinyUAV(FRHIUnorderedAccessView* UnorderedAccessViewRHI, const uint32* Values) final override;
-	virtual void RHICopyToResolveTarget(FTextureRHIParamRef SourceTexture, FTextureRHIParamRef DestTexture, const FResolveParams& ResolveParams) final override;
-	virtual void RHICopyTexture(FTextureRHIParamRef SourceTexture, FTextureRHIParamRef DestTexture, const FRHICopyTextureInfo& CopyInfo) final override;
+	virtual void RHICopyToResolveTarget(FRHITexture* SourceTexture, FRHITexture* DestTexture, const FResolveParams& ResolveParams) final override;
+	virtual void RHICopyTexture(FRHITexture* SourceTexture, FRHITexture* DestTexture, const FRHICopyTextureInfo& CopyInfo) final override;
 	virtual void RHIBeginRenderQuery(FRHIRenderQuery* RenderQuery) final override;
 	virtual void RHIEndRenderQuery(FRHIRenderQuery* RenderQuery) final override;
 	virtual void RHISubmitCommandsHint() final override;
-	virtual void RHIBeginDrawingViewport(FRHIViewport* Viewport, FTextureRHIParamRef RenderTargetRHI) final override;
+	virtual void RHIBeginDrawingViewport(FRHIViewport* Viewport, FRHITexture* RenderTargetRHI) final override;
 	virtual void RHIEndDrawingViewport(FRHIViewport* Viewport, bool bPresent, bool bLockToVsync) final override;
 	virtual void RHIBeginFrame() final override;
 	virtual void RHIEndFrame() final override;
 	virtual void RHIBeginScene() final override;
 	virtual void RHIEndScene() final override;
-	virtual void RHISetStreamSource(uint32 StreamIndex, FVertexBufferRHIParamRef VertexBuffer, uint32 Offset) final override;
+	virtual void RHISetStreamSource(uint32 StreamIndex, FRHIVertexBuffer* VertexBuffer, uint32 Offset) final override;
 	virtual void RHISetRasterizerState(FRHIRasterizerState* NewState) final override;
 	virtual void RHISetViewport(uint32 MinX, uint32 MinY, float MinZ, uint32 MaxX, uint32 MaxY, float MaxZ) final override;
 	virtual void RHISetScissorRect(bool bEnable, uint32 MinX, uint32 MinY, uint32 MaxX, uint32 MaxY) final override;
 	virtual void RHISetBoundShaderState(FRHIBoundShaderState* BoundShaderState) final override;
-	virtual void RHISetShaderTexture(FRHIVertexShader* VertexShader, uint32 TextureIndex, FTextureRHIParamRef NewTexture) final override;
-	virtual void RHISetShaderTexture(FRHIHullShader* HullShader, uint32 TextureIndex, FTextureRHIParamRef NewTexture) final override;
-	virtual void RHISetShaderTexture(FRHIDomainShader* DomainShader, uint32 TextureIndex, FTextureRHIParamRef NewTexture) final override;
-	virtual void RHISetShaderTexture(FRHIGeometryShader* GeometryShader, uint32 TextureIndex, FTextureRHIParamRef NewTexture) final override;
-	virtual void RHISetShaderTexture(FRHIPixelShader* PixelShader, uint32 TextureIndex, FTextureRHIParamRef NewTexture) final override;
-	virtual void RHISetShaderTexture(FRHIComputeShader* PixelShader, uint32 TextureIndex, FTextureRHIParamRef NewTexture) final override;
+	virtual void RHISetShaderTexture(FRHIVertexShader* VertexShader, uint32 TextureIndex, FRHITexture* NewTexture) final override;
+	virtual void RHISetShaderTexture(FRHIHullShader* HullShader, uint32 TextureIndex, FRHITexture* NewTexture) final override;
+	virtual void RHISetShaderTexture(FRHIDomainShader* DomainShader, uint32 TextureIndex, FRHITexture* NewTexture) final override;
+	virtual void RHISetShaderTexture(FRHIGeometryShader* GeometryShader, uint32 TextureIndex, FRHITexture* NewTexture) final override;
+	virtual void RHISetShaderTexture(FRHIPixelShader* PixelShader, uint32 TextureIndex, FRHITexture* NewTexture) final override;
+	virtual void RHISetShaderTexture(FRHIComputeShader* PixelShader, uint32 TextureIndex, FRHITexture* NewTexture) final override;
 	virtual void RHISetShaderSampler(FRHIComputeShader* ComputeShader, uint32 SamplerIndex, FRHISamplerState* NewState) final override;
 	virtual void RHISetShaderSampler(FRHIVertexShader* VertexShader, uint32 SamplerIndex, FRHISamplerState* NewState) final override;
 	virtual void RHISetShaderSampler(FRHIGeometryShader* GeometryShader, uint32 SamplerIndex, FRHISamplerState* NewState) final override;
@@ -175,12 +175,12 @@ public:
 	virtual void RHISetShaderResourceViewParameter(FRHIHullShader* HullShader, uint32 SamplerIndex, FRHIShaderResourceView* SRV) final override;
 	virtual void RHISetShaderResourceViewParameter(FRHIDomainShader* DomainShader, uint32 SamplerIndex, FRHIShaderResourceView* SRV) final override;
 	virtual void RHISetShaderResourceViewParameter(FRHIGeometryShader* GeometryShader, uint32 SamplerIndex, FRHIShaderResourceView* SRV) final override;
-	virtual void RHISetShaderUniformBuffer(FRHIVertexShader* VertexShader, uint32 BufferIndex, FUniformBufferRHIParamRef Buffer) final override;
-	virtual void RHISetShaderUniformBuffer(FRHIHullShader* HullShader, uint32 BufferIndex, FUniformBufferRHIParamRef Buffer) final override;
-	virtual void RHISetShaderUniformBuffer(FRHIDomainShader* DomainShader, uint32 BufferIndex, FUniformBufferRHIParamRef Buffer) final override;
-	virtual void RHISetShaderUniformBuffer(FRHIGeometryShader* GeometryShader, uint32 BufferIndex, FUniformBufferRHIParamRef Buffer) final override;
-	virtual void RHISetShaderUniformBuffer(FRHIPixelShader* PixelShader, uint32 BufferIndex, FUniformBufferRHIParamRef Buffer) final override;
-	virtual void RHISetShaderUniformBuffer(FRHIComputeShader* ComputeShader, uint32 BufferIndex, FUniformBufferRHIParamRef Buffer) final override;
+	virtual void RHISetShaderUniformBuffer(FRHIVertexShader* VertexShader, uint32 BufferIndex, FRHIUniformBuffer* Buffer) final override;
+	virtual void RHISetShaderUniformBuffer(FRHIHullShader* HullShader, uint32 BufferIndex, FRHIUniformBuffer* Buffer) final override;
+	virtual void RHISetShaderUniformBuffer(FRHIDomainShader* DomainShader, uint32 BufferIndex, FRHIUniformBuffer* Buffer) final override;
+	virtual void RHISetShaderUniformBuffer(FRHIGeometryShader* GeometryShader, uint32 BufferIndex, FRHIUniformBuffer* Buffer) final override;
+	virtual void RHISetShaderUniformBuffer(FRHIPixelShader* PixelShader, uint32 BufferIndex, FRHIUniformBuffer* Buffer) final override;
+	virtual void RHISetShaderUniformBuffer(FRHIComputeShader* ComputeShader, uint32 BufferIndex, FRHIUniformBuffer* Buffer) final override;
 	virtual void RHISetShaderParameter(FRHIVertexShader* VertexShader, uint32 BufferIndex, uint32 BaseIndex, uint32 NumBytes, const void* NewValue) final override;
 	virtual void RHISetShaderParameter(FRHIPixelShader* PixelShader, uint32 BufferIndex, uint32 BaseIndex, uint32 NumBytes, const void* NewValue) final override;
 	virtual void RHISetShaderParameter(FRHIHullShader* HullShader, uint32 BufferIndex, uint32 BaseIndex, uint32 NumBytes, const void* NewValue) final override;
@@ -192,10 +192,10 @@ public:
 	virtual void RHISetRenderTargets(uint32 NumSimultaneousRenderTargets, const FRHIRenderTargetView* NewRenderTargets, const FRHIDepthRenderTargetView* NewDepthStencilTarget, uint32 NumUAVs, FRHIUnorderedAccessView* const* UAVs) final override;
 	virtual void RHISetRenderTargetsAndClear(const FRHISetRenderTargetsInfo& RenderTargetsInfo) final override;
 	virtual void RHIDrawPrimitive(uint32 PrimitiveType, uint32 BaseVertexIndex, uint32 NumPrimitives, uint32 NumInstances) final override;
-	virtual void RHIDrawPrimitiveIndirect(uint32 PrimitiveType, FVertexBufferRHIParamRef ArgumentBuffer, uint32 ArgumentOffset) final override;
-	virtual void RHIDrawIndexedIndirect(FIndexBufferRHIParamRef IndexBufferRHI, uint32 PrimitiveType, FStructuredBufferRHIParamRef ArgumentsBufferRHI, int32 DrawArgumentsIndex, uint32 NumInstances) final override;
-	virtual void RHIDrawIndexedPrimitive(FIndexBufferRHIParamRef IndexBuffer, uint32 PrimitiveType, int32 BaseVertexIndex, uint32 FirstInstance, uint32 NumVertices, uint32 StartIndex, uint32 NumPrimitives, uint32 NumInstances) final override;
-	virtual void RHIDrawIndexedPrimitiveIndirect(uint32 PrimitiveType, FIndexBufferRHIParamRef IndexBuffer, FVertexBufferRHIParamRef ArgumentBuffer, uint32 ArgumentOffset) final override;
+	virtual void RHIDrawPrimitiveIndirect(uint32 PrimitiveType, FRHIVertexBuffer* ArgumentBuffer, uint32 ArgumentOffset) final override;
+	virtual void RHIDrawIndexedIndirect(FRHIIndexBuffer* IndexBufferRHI, uint32 PrimitiveType, FRHIStructuredBuffer* ArgumentsBufferRHI, int32 DrawArgumentsIndex, uint32 NumInstances) final override;
+	virtual void RHIDrawIndexedPrimitive(FRHIIndexBuffer* IndexBuffer, uint32 PrimitiveType, int32 BaseVertexIndex, uint32 FirstInstance, uint32 NumVertices, uint32 StartIndex, uint32 NumPrimitives, uint32 NumInstances) final override;
+	virtual void RHIDrawIndexedPrimitiveIndirect(uint32 PrimitiveType, FRHIIndexBuffer* IndexBuffer, FRHIVertexBuffer* ArgumentBuffer, uint32 ArgumentOffset) final override;
 	virtual void RHIEnableDepthBoundsTest(bool bEnable) final override;
 	virtual void RHISetDepthBounds(float MinDepth, float MaxDepth) final override;
 	virtual void RHIPushEvent(const TCHAR* Name, FColor Color) final override;
diff --git a/Engine/Source/Runtime/Engine/Private/Components/ReflectionCaptureComponent.cpp b/Engine/Source/Runtime/Engine/Private/Components/ReflectionCaptureComponent.cpp
--- a/Engine/Source/Runtime/Engine/Private/Components/ReflectionCaptureComponent.cpp
+++ b/Engine/Source/Runtime/Engine/Private/Components/ReflectionCaptureComponent.cpp
@@ -702,7 +702,7 @@ public:
 		return Size;
 	}
 
-	FTextureRHIParamRef GetTextureRHI() 
+	FRHITexture* GetTextureRHI() 
 	{
 		return TextureCubeRHI;
 	}
diff --git a/Engine/Source/Runtime/Engine/Private/GPUSkinCache.cpp b/Engine/Source/Runtime/Engine/Private/GPUSkinCache.cpp
--- a/Engine/Source/Runtime/Engine/Private/GPUSkinCache.cpp
+++ b/Engine/Source/Runtime/Engine/Private/GPUSkinCache.cpp
@@ -883,7 +883,7 @@ void FGPUSkinCache::DispatchUpdateSkinTangents(FRHICommandListImmediate& RHICmdL
 		//SetRenderTarget(RHICmdList, FTextureRHIRef(), FTextureRHIRef());
 
 		FRawStaticIndexBuffer16or32Interface* IndexBuffer = LodData.MultiSizeIndexContainer.GetIndexBuffer();
-		FIndexBufferRHIParamRef IndexBufferRHI = IndexBuffer->IndexBufferRHI;
+		FRHIIndexBuffer* IndexBufferRHI = IndexBuffer->IndexBufferRHI;
 
 		const uint32 RequiredVertexCount = LodData.GetNumVertices();
 
diff --git a/Engine/Source/Runtime/Engine/Private/GPUSkinVertexFactory.cpp b/Engine/Source/Runtime/Engine/Private/GPUSkinVertexFactory.cpp
--- a/Engine/Source/Runtime/Engine/Private/GPUSkinVertexFactory.cpp
+++ b/Engine/Source/Runtime/Engine/Private/GPUSkinVertexFactory.cpp
@@ -199,13 +199,13 @@ static bool DeferSkeletalLockAndFillToRHIThread()
 
 struct FRHICommandUpdateBoneBuffer final : public FRHICommand<FRHICommandUpdateBoneBuffer>
 {
-	FVertexBufferRHIParamRef VertexBuffer;
+	FRHIVertexBuffer* VertexBuffer;
 	uint32 BufferSize;
 	const TArray<FMatrix>& ReferenceToLocalMatrices;
 	const TArray<FBoneIndexType>& BoneMap;
 
 
-	FORCEINLINE_DEBUGGABLE FRHICommandUpdateBoneBuffer(FVertexBufferRHIParamRef InVertexBuffer, uint32 InBufferSize, const TArray<FMatrix>& InReferenceToLocalMatrices, const TArray<FBoneIndexType>& InBoneMap)
+	FORCEINLINE_DEBUGGABLE FRHICommandUpdateBoneBuffer(FRHIVertexBuffer* InVertexBuffer, uint32 InBufferSize, const TArray<FMatrix>& InReferenceToLocalMatrices, const TArray<FBoneIndexType>& InBoneMap)
 		: VertexBuffer(InVertexBuffer)
 		, BufferSize(InBufferSize)
 		, ReferenceToLocalMatrices(InReferenceToLocalMatrices)
@@ -623,7 +623,7 @@ public:
 		check(BatchUserData);
 
 		const auto* LocalVertexFactory = static_cast<const FGPUSkinPassthroughVertexFactory*>(VertexFactory);
-		FUniformBufferRHIParamRef VertexFactoryUniformBuffer = nullptr;
+		FRHIUniformBuffer* VertexFactoryUniformBuffer = nullptr;
 		VertexFactoryUniformBuffer = LocalVertexFactory->GetUniformBuffer();
 
 		// #dxr_todo do we need this call to the base?
@@ -899,13 +899,13 @@ protected:
 
 struct FRHICommandUpdateClothBuffer final : public FRHICommand<FRHICommandUpdateClothBuffer>
 {
-	FVertexBufferRHIParamRef VertexBuffer;
+	FRHIVertexBuffer* VertexBuffer;
 	uint32 BufferSize;
 	const TArray<FVector>& SimulPositions;
 	const TArray<FVector>& SimulNormals;
 
 
-	FORCEINLINE_DEBUGGABLE FRHICommandUpdateClothBuffer(FVertexBufferRHIParamRef InVertexBuffer, uint32 InBufferSize, const TArray<FVector>& InSimulPositions, const TArray<FVector>& InSimulNormals)
+	FORCEINLINE_DEBUGGABLE FRHICommandUpdateClothBuffer(FRHIVertexBuffer* InVertexBuffer, uint32 InBufferSize, const TArray<FVector>& InSimulPositions, const TArray<FVector>& InSimulNormals)
 		: VertexBuffer(InVertexBuffer)
 		, BufferSize(InBufferSize)
 		, SimulPositions(InSimulPositions)
diff --git a/Engine/Source/Runtime/Engine/Private/GeometryCollectionRendering.cpp b/Engine/Source/Runtime/Engine/Private/GeometryCollectionRendering.cpp
--- a/Engine/Source/Runtime/Engine/Private/GeometryCollectionRendering.cpp
+++ b/Engine/Source/Runtime/Engine/Private/GeometryCollectionRendering.cpp
@@ -20,7 +20,7 @@ void FGeometryCollectionVertexFactoryShaderParameters::GetElementShaderBindings(
 
 
 	const auto* LocalVertexFactory = static_cast<const FGeometryCollectionVertexFactory*>(VertexFactory);
-	FUniformBufferRHIParamRef VertexFactoryUniformBuffer = nullptr;
+	FRHIUniformBuffer* VertexFactoryUniformBuffer = nullptr;
 	VertexFactoryUniformBuffer = LocalVertexFactory->GetUniformBuffer();
 
 
diff --git a/Engine/Source/Runtime/Engine/Private/InstancedStaticMesh.cpp b/Engine/Source/Runtime/Engine/Private/InstancedStaticMesh.cpp
--- a/Engine/Source/Runtime/Engine/Private/InstancedStaticMesh.cpp
+++ b/Engine/Source/Runtime/Engine/Private/InstancedStaticMesh.cpp
@@ -2899,7 +2899,7 @@ void FInstancedStaticMeshVertexFactoryShaderParameters::GetElementShaderBindings
 	) const
 {
 	// Decode VertexFactoryUserData as VertexFactoryUniformBuffer
-	FUniformBufferRHIParamRef VertexFactoryUniformBuffer = static_cast<FUniformBufferRHIParamRef>(BatchElement.VertexFactoryUserData);
+	FRHIUniformBuffer* VertexFactoryUniformBuffer = static_cast<FRHIUniformBuffer*>(BatchElement.VertexFactoryUserData);
 	FLocalVertexFactoryShaderParametersBase::GetElementShaderBindingsBase(Scene, View, Shader, InputStreamType, FeatureLevel, VertexFactory, BatchElement, VertexFactoryUniformBuffer, ShaderBindings, VertexStreams);
 
 	const FInstancingUserData* InstancingUserData = (const FInstancingUserData*)BatchElement.UserData;
diff --git a/Engine/Source/Runtime/Engine/Private/LocalVertexFactory.cpp b/Engine/Source/Runtime/Engine/Private/LocalVertexFactory.cpp
--- a/Engine/Source/Runtime/Engine/Private/LocalVertexFactory.cpp
+++ b/Engine/Source/Runtime/Engine/Private/LocalVertexFactory.cpp
@@ -98,7 +98,7 @@ void FLocalVertexFactoryShaderParametersBase::GetElementShaderBindingsBase(
 	ERHIFeatureLevel::Type FeatureLevel,
 	const FVertexFactory* VertexFactory, 
 	const FMeshBatchElement& BatchElement,
-	FUniformBufferRHIParamRef VertexFactoryUniformBuffer,
+	FRHIUniformBuffer* VertexFactoryUniformBuffer,
 	FMeshDrawSingleShaderBindings& ShaderBindings,
 	FVertexInputStreamArray& VertexStreams
 	) const
@@ -131,7 +131,7 @@ void FLocalVertexFactoryShaderParametersBase::GetElementShaderBindingsBase(
 	if (bAnySpeedTreeParamIsBound)
 	{
 		QUICK_SCOPE_CYCLE_COUNTER(STAT_FLocalVertexFactoryShaderParameters_SetMesh_SpeedTree);
-		FUniformBufferRHIParamRef SpeedTreeUniformBuffer = Scene? Scene->GetSpeedTreeUniformBuffer(VertexFactory) : nullptr;
+		FRHIUniformBuffer* SpeedTreeUniformBuffer = Scene? Scene->GetSpeedTreeUniformBuffer(VertexFactory) : nullptr;
 		if (SpeedTreeUniformBuffer == nullptr)
 		{
 			SpeedTreeUniformBuffer = GSpeedTreeWindNullUniformBuffer.GetUniformBufferRHI();
@@ -161,7 +161,7 @@ void FLocalVertexFactoryShaderParameters::GetElementShaderBindings(
 ) const
 {
 	// Decode VertexFactoryUserData as VertexFactoryUniformBuffer
-	FUniformBufferRHIParamRef VertexFactoryUniformBuffer = static_cast<FUniformBufferRHIParamRef>(BatchElement.VertexFactoryUserData);
+	FRHIUniformBuffer* VertexFactoryUniformBuffer = static_cast<FRHIUniformBuffer*>(BatchElement.VertexFactoryUserData);
 
 	FLocalVertexFactoryShaderParametersBase::GetElementShaderBindingsBase(
 		Scene,
diff --git a/Engine/Source/Runtime/Engine/Private/MeshParticleVertexFactory.cpp b/Engine/Source/Runtime/Engine/Private/MeshParticleVertexFactory.cpp
--- a/Engine/Source/Runtime/Engine/Private/MeshParticleVertexFactory.cpp
+++ b/Engine/Source/Runtime/Engine/Private/MeshParticleVertexFactory.cpp
@@ -137,7 +137,7 @@ public:
 		return TEXT("FDummyPrevTransformBuffer");
 	}
 
-	inline FVertexBufferRHIParamRef GetVB() const
+	inline FRHIVertexBuffer* GetVB() const
 	{
 		return VB;
 	}
diff --git a/Engine/Source/Runtime/Engine/Private/ModelRender.cpp b/Engine/Source/Runtime/Engine/Private/ModelRender.cpp
--- a/Engine/Source/Runtime/Engine/Private/ModelRender.cpp
+++ b/Engine/Source/Runtime/Engine/Private/ModelRender.cpp
@@ -724,7 +724,7 @@ private:
 		// Accessors.
 		UMaterialInterface* GetMaterial() const { return Material; }
 
-		FUniformBufferRHIParamRef GetVertexFactoryUniformBuffer() const { return VertexFactoryUniformBuffer; }
+		FRHIUniformBuffer* GetVertexFactoryUniformBuffer() const { return VertexFactoryUniformBuffer; }
 
 	private:
 
diff --git a/Engine/Source/Runtime/Engine/Private/Particles/FXSystem.cpp b/Engine/Source/Runtime/Engine/Private/Particles/FXSystem.cpp
--- a/Engine/Source/Runtime/Engine/Private/Particles/FXSystem.cpp
+++ b/Engine/Source/Runtime/Engine/Private/Particles/FXSystem.cpp
@@ -398,7 +398,7 @@ void FFXSystem::PreRender(FRHICommandListImmediate& RHICmdList, const FGlobalDis
 		PrepareGPUSimulation(RHICmdList);
 
 		RHICmdList.SetCurrentStat(GET_STATID(STAT_CLM_FXPreRender_Simulate));
-		SimulateGPUParticles(RHICmdList, EParticleSimulatePhase::Main, nullptr, nullptr, nullptr, FUniformBufferRHIParamRef());
+		SimulateGPUParticles(RHICmdList, EParticleSimulatePhase::Main, nullptr, nullptr, nullptr, nullptr);
 
 		RHICmdList.SetCurrentStat(GET_STATID(STAT_CLM_FXPreRender_Finalize));
 		FinalizeGPUSimulation(RHICmdList);
@@ -409,7 +409,7 @@ void FFXSystem::PreRender(FRHICommandListImmediate& RHICmdList, const FGlobalDis
 			PrepareGPUSimulation(RHICmdList);
 
 			RHICmdList.SetCurrentStat(GET_STATID(STAT_CLM_FXPreRender_SimulateCDF));
-			SimulateGPUParticles(RHICmdList, EParticleSimulatePhase::CollisionDistanceField, nullptr, GlobalDistanceFieldParameterData, nullptr, FUniformBufferRHIParamRef());
+			SimulateGPUParticles(RHICmdList, EParticleSimulatePhase::CollisionDistanceField, nullptr, GlobalDistanceFieldParameterData, nullptr, nullptr);
 			//particles rendered during basepass may need to read pos/velocity buffers.  must finalize unless we know for sure that nothingin base pass will read.
 			RHICmdList.SetCurrentStat(GET_STATID(STAT_CLM_FXPreRender_FinalizeCDF));
 			FinalizeGPUSimulation(RHICmdList);
@@ -419,9 +419,9 @@ void FFXSystem::PreRender(FRHICommandListImmediate& RHICmdList, const FGlobalDis
 
 void FFXSystem::PostRenderOpaque(
 	FRHICommandListImmediate& RHICmdList, 
-	const FUniformBufferRHIParamRef ViewUniformBuffer, 
+	FRHIUniformBuffer* ViewUniformBuffer,
 	const FShaderParametersMetadata* SceneTexturesUniformBufferStruct,
-	FUniformBufferRHIParamRef SceneTexturesUniformBuffer)
+	FRHIUniformBuffer* SceneTexturesUniformBuffer)
 {
 	if (RHISupportsGPUParticles() && IsParticleCollisionModeSupported(GetShaderPlatform(), PCM_DepthBuffer))
 	{
diff --git a/Engine/Source/Runtime/Engine/Private/Particles/FXSystemPrivate.h b/Engine/Source/Runtime/Engine/Private/Particles/FXSystemPrivate.h
--- a/Engine/Source/Runtime/Engine/Private/Particles/FXSystemPrivate.h
+++ b/Engine/Source/Runtime/Engine/Private/Particles/FXSystemPrivate.h
@@ -154,9 +154,9 @@ public:
 	virtual void PreRender(FRHICommandListImmediate& RHICmdList, const FGlobalDistanceFieldParameterData* GlobalDistanceFieldParameterData) override;
 	virtual void PostRenderOpaque(
 		FRHICommandListImmediate& RHICmdList, 
-		const FUniformBufferRHIParamRef ViewUniformBuffer, 
+		FRHIUniformBuffer* ViewUniformBuffer,
 		const FShaderParametersMetadata* SceneTexturesUniformBufferStruct,
-		FUniformBufferRHIParamRef SceneTexturesUniformBuffer) override;
+		FRHIUniformBuffer* SceneTexturesUniformBuffer) override;
 	// End FFXSystemInterface.
 
 	/*--------------------------------------------------------------------------
@@ -255,10 +255,10 @@ private:
 	void SimulateGPUParticles(
 		FRHICommandListImmediate& RHICmdList,
 		EParticleSimulatePhase::Type Phase,
-		const FUniformBufferRHIParamRef ViewUniformBuffer,
+		FRHIUniformBuffer* ViewUniformBuffer,
 		const FGlobalDistanceFieldParameterData* GlobalDistanceFieldParameterData,
 		const FShaderParametersMetadata* SceneTexturesUniformBufferStruct,
-		FUniformBufferRHIParamRef SceneTexturesUniformBuffer
+		FRHIUniformBuffer* SceneTexturesUniformBuffer
 		);
 
 	/**
@@ -273,10 +273,10 @@ private:
 	void SimulateGPUParticles_Internal(
 		FRHICommandListImmediate& RHICmdList,
 		EParticleSimulatePhase::Type Phase,
-		const FUniformBufferRHIParamRef ViewUniformBuffer,
+		FRHIUniformBuffer* ViewUniformBuffer,
 		const FGlobalDistanceFieldParameterData* GlobalDistanceFieldParameterData,
-		FTexture2DRHIParamRef SceneDepthTexture,
-		FTexture2DRHIParamRef GBufferATexture
+		FRHITexture2D* SceneDepthTexture,
+		FRHITexture2D* GBufferATexture
 	);
 
 	/*-------------------------------------------------------------------------
diff --git a/Engine/Source/Runtime/Engine/Private/Particles/FXSystemSet.cpp b/Engine/Source/Runtime/Engine/Private/Particles/FXSystemSet.cpp
--- a/Engine/Source/Runtime/Engine/Private/Particles/FXSystemSet.cpp
+++ b/Engine/Source/Runtime/Engine/Private/Particles/FXSystemSet.cpp
@@ -120,9 +120,9 @@ void FFXSystemSet::PreRender(FRHICommandListImmediate& RHICmdList, const class F
 
 void FFXSystemSet::PostRenderOpaque(
 	FRHICommandListImmediate& RHICmdList,
-	const FUniformBufferRHIParamRef ViewUniformBuffer,
+	FRHIUniformBuffer* ViewUniformBuffer,
 	const class FShaderParametersMetadata* SceneTexturesUniformBufferStruct,
-	FUniformBufferRHIParamRef SceneTexturesUniformBuffer)
+	FRHIUniformBuffer* SceneTexturesUniformBuffer)
 {
 	for (FFXSystemInterface* FXSystem : FXSystems)
 	{
diff --git a/Engine/Source/Runtime/Engine/Private/Particles/FXSystemSet.h b/Engine/Source/Runtime/Engine/Private/Particles/FXSystemSet.h
--- a/Engine/Source/Runtime/Engine/Private/Particles/FXSystemSet.h
+++ b/Engine/Source/Runtime/Engine/Private/Particles/FXSystemSet.h
@@ -35,9 +35,9 @@ public:
 	virtual void PreRender(FRHICommandListImmediate& RHICmdList, const class FGlobalDistanceFieldParameterData* GlobalDistanceFieldParameterData) override;
 	virtual void PostRenderOpaque(
 		FRHICommandListImmediate& RHICmdList, 
-		const FUniformBufferRHIParamRef ViewUniformBuffer, 
+		FRHIUniformBuffer* ViewUniformBuffer,
 		const class FShaderParametersMetadata* SceneTexturesUniformBufferStruct,
-		FUniformBufferRHIParamRef SceneTexturesUniformBuffer) override;
+		FRHIUniformBuffer* SceneTexturesUniformBuffer) override;
 
 protected:
 
diff --git a/Engine/Source/Runtime/Engine/Private/Particles/ParticleCurveTexture.cpp b/Engine/Source/Runtime/Engine/Private/Particles/ParticleCurveTexture.cpp
--- a/Engine/Source/Runtime/Engine/Private/Particles/ParticleCurveTexture.cpp
+++ b/Engine/Source/Runtime/Engine/Private/Particles/ParticleCurveTexture.cpp
@@ -175,8 +175,8 @@ TGlobalResource<FParticleCurveInjectionVertexDeclaration> GParticleCurveInjectio
 */
 static void InjectCurves(
 	FRHICommandListImmediate& RHICmdList,
-	FTexture2DRHIParamRef CurveTextureRHI,
-	FTexture2DRHIParamRef CurveTextureTargetRHI,
+	FRHITexture2D* CurveTextureRHI,
+	FRHITexture2D* CurveTextureTargetRHI,
 	TArray<FCurveSamples>& InPendingCurves)
 {
 	static bool bFirstCall = true;
diff --git a/Engine/Source/Runtime/Engine/Private/Particles/ParticleCurveTexture.h b/Engine/Source/Runtime/Engine/Private/Particles/ParticleCurveTexture.h
--- a/Engine/Source/Runtime/Engine/Private/Particles/ParticleCurveTexture.h
+++ b/Engine/Source/Runtime/Engine/Private/Particles/ParticleCurveTexture.h
@@ -158,7 +158,7 @@ public:
 	/**
 	 * Retrieves the curve texture from which shaders can sample.
 	 */
-	FTexture2DRHIParamRef GetCurveTexture() { return CurveTextureRHI; }
+	FRHITexture2D* GetCurveTexture() { return CurveTextureRHI; }
 
 	/**
 	 * Submits pending curves to the GPU.
diff --git a/Engine/Source/Runtime/Engine/Private/Particles/ParticleGpuSimulation.cpp b/Engine/Source/Runtime/Engine/Private/Particles/ParticleGpuSimulation.cpp
--- a/Engine/Source/Runtime/Engine/Private/Particles/ParticleGpuSimulation.cpp
+++ b/Engine/Source/Runtime/Engine/Private/Particles/ParticleGpuSimulation.cpp
@@ -406,7 +406,7 @@ public:
 			ParticleResources->SimulationAttributesTexture.InitResource();
 			ParticleResources->SortedVertexBuffer.InitResource();
 
-			FTextureRHIParamRef AttributeTextures[2];
+			FRHITexture* AttributeTextures[2];
 			AttributeTextures[0] = ParticleResources->RenderAttributesTexture.TextureRHI;
 			AttributeTextures[1] = ParticleResources->SimulationAttributesTexture.TextureRHI;
 			RHICmdList.TransitionResources(EResourceTransitionAccess::EReadable, AttributeTextures, 2);
@@ -697,17 +697,17 @@ public:
 	}
 
 	/** Emitter uniform buffer. */
-	FUniformBufferRHIParamRef EmitterUniformBuffer;
+	FRHIUniformBuffer* EmitterUniformBuffer;
 	/** Emitter uniform buffer for dynamic parameters. */
 	FUniformBufferRHIRef EmitterDynamicUniformBuffer;
 	/** Buffer containing unsorted particle indices. */
 	FRHIShaderResourceView* UnsortedParticleIndicesSRV;
 	/** Texture containing positions for all particles. */
-	FTexture2DRHIParamRef PositionTextureRHI;
+	FRHITexture2D* PositionTextureRHI;
 	/** Texture containing velocities for all particles. */
-	FTexture2DRHIParamRef VelocityTextureRHI;
+	FRHITexture2D* VelocityTextureRHI;
 	/** Texture containint attributes for all particles. */
-	FTexture2DRHIParamRef AttributesTextureRHI;
+	FRHITexture2D* AttributesTextureRHI;
 
 
 	FGPUSpriteVertexFactory()
@@ -1150,11 +1150,11 @@ public:
 		const FParticleStateTextures& TextureResources,
 		const FParticleAttributesTexture& InAttributesTexture,
 		const FParticleAttributesTexture& InRenderAttributesTexture,
-		const FUniformBufferRHIParamRef ViewUniformBuffer,
+		FRHIUniformBuffer* ViewUniformBuffer,
 		ERHIFeatureLevel::Type FeatureLevel,
 		const FGlobalDistanceFieldParameterData* GlobalDistanceFieldParameterData,
 		const FShaderParametersMetadata* SceneTexturesUniformBufferStruct,
-		FUniformBufferRHIParamRef SceneTexturesUniformBuffer
+		FRHIUniformBuffer* SceneTexturesUniformBuffer
 		)
 	{
 		FRHIPixelShader* PixelShaderRHI = GetPixelShader();
@@ -1200,7 +1200,7 @@ public:
 	 * Set parameters for the vector fields sampled by this shader.
 	 * @param VectorFieldParameters -Parameters needed to sample local vector fields.
 	 */
-	void SetVectorFieldParameters(FRHICommandList& RHICmdList, const FVectorFieldUniformBufferRef& UniformBuffer, const FTexture3DRHIParamRef VolumeTexturesRHI[])
+	void SetVectorFieldParameters(FRHICommandList& RHICmdList, const FVectorFieldUniformBufferRef& UniformBuffer, FRHITexture3D* const* VolumeTexturesRHI)
 	{
 		FRHIPixelShader* PixelShaderRHI = GetPixelShader();
 			SetUniformBufferParameter(RHICmdList, PixelShaderRHI, GetUniformBufferParameter<FVectorFieldUniformParameters>(), UniformBuffer);
@@ -1217,7 +1217,7 @@ public:
 	/**
 	 * Set per-instance parameters for this shader.
 	 */
-	void SetInstanceParameters(FRHICommandList& RHICmdList, FUniformBufferRHIParamRef UniformBuffer, const FParticlePerFrameSimulationParameters& InPerFrameParameters, bool bUseFixDT)
+	void SetInstanceParameters(FRHICommandList& RHICmdList, FRHIUniformBuffer* UniformBuffer, const FParticlePerFrameSimulationParameters& InPerFrameParameters, bool bUseFixDT)
 	{
 		FRHIPixelShader* PixelShaderRHI = GetPixelShader();
 		SetUniformBufferParameter(RHICmdList, PixelShaderRHI, GetUniformBufferParameter<FParticleSimulationParameters>(), UniformBuffer);
@@ -1482,19 +1482,19 @@ struct FSimulationCommandGPU
 	FParticleShaderParamRef TileOffsetsShaderRef;
 	FParticleBufferParamRef TileOffsetsBufferRef;
 	/** Uniform buffer containing simulation parameters. */
-	FUniformBufferRHIParamRef UniformBuffer;
+	FRHIUniformBuffer* UniformBuffer;
 	/** Uniform buffer containing per-frame simulation parameters. */
 	FParticlePerFrameSimulationParameters PerFrameParameters;
 	/** Parameters to sample the local vector field for this simulation. */
 	FVectorFieldUniformBufferRef VectorFieldsUniformBuffer;
 	FLocalUniformBuffer VectorFieldsUniformBufferLocal;
 	/** Vector field volume textures for this simulation. */
-	FTexture3DRHIParamRef VectorFieldTexturesRHI[MAX_VECTOR_FIELDS];
+	FRHITexture3D* VectorFieldTexturesRHI[MAX_VECTOR_FIELDS];
 	/** The number of tiles to simulate. */
 	int32 UnalignedTileCount;
 
 	/** Initialization constructor. */
-	FSimulationCommandGPU(FParticleShaderParamRef InTileOffsetsShaderRef, FParticleBufferParamRef InTileOffsetsBufferRef, FUniformBufferRHIParamRef InUniformBuffer, const FParticlePerFrameSimulationParameters& InPerFrameParameters, FVectorFieldUniformBufferRef& InVectorFieldsUniformBuffer, int32 InTileCount)
+	FSimulationCommandGPU(FParticleShaderParamRef InTileOffsetsShaderRef, FParticleBufferParamRef InTileOffsetsBufferRef, FRHIUniformBuffer* InUniformBuffer, const FParticlePerFrameSimulationParameters& InPerFrameParameters, FVectorFieldUniformBufferRef& InVectorFieldsUniformBuffer, int32 InTileCount)
 		: TileOffsetsShaderRef(InTileOffsetsShaderRef)
 		, TileOffsetsBufferRef(InTileOffsetsBufferRef)
 		, UniformBuffer(InUniformBuffer)
@@ -1502,7 +1502,7 @@ struct FSimulationCommandGPU
 		, VectorFieldsUniformBuffer(InVectorFieldsUniformBuffer)
 		, UnalignedTileCount(InTileCount)
 	{
-		FTexture3DRHIParamRef BlackVolumeTextureRHI = (FTexture3DRHIParamRef)(FTextureRHIParamRef)GBlackVolumeTexture->TextureRHI;
+		FRHITexture3D* BlackVolumeTextureRHI = (FRHITexture3D*)(FRHITexture*)GBlackVolumeTexture->TextureRHI;
 		for (int32 i = 0; i < MAX_VECTOR_FIELDS; ++i)
 		{
 			VectorFieldTexturesRHI[i] = BlackVolumeTextureRHI;
@@ -1525,10 +1525,10 @@ void ExecuteSimulationCommands(
 	ERHIFeatureLevel::Type FeatureLevel,
 	const TArray<FSimulationCommandGPU>& SimulationCommands,
 	FParticleSimulationResources* ParticleSimulationResources,
-	const FUniformBufferRHIParamRef ViewUniformBuffer,
+	FRHIUniformBuffer* ViewUniformBuffer,
 	const FGlobalDistanceFieldParameterData* GlobalDistanceFieldParameterData,
 	const FShaderParametersMetadata* SceneTexturesUniformBufferStruct,
-	FUniformBufferRHIParamRef SceneTexturesUniformBuffer,
+	FRHIUniformBuffer* SceneTexturesUniformBuffer,
 	bool bUseFixDT)
 {
 	if (!CVarSimulateGPUParticles.GetValueOnAnyThread())
@@ -1594,10 +1594,10 @@ void ExecuteSimulationCommands(
 	ERHIFeatureLevel::Type FeatureLevel,
 	const TArray<FSimulationCommandGPU>& SimulationCommands,
 	FParticleSimulationResources* ParticleSimulationResources,
-	const FUniformBufferRHIParamRef ViewUniformBuffer,
+	FRHIUniformBuffer* ViewUniformBuffer,
 	const FGlobalDistanceFieldParameterData* GlobalDistanceFieldParameterData,
 	const FShaderParametersMetadata* SceneTexturesUniformBufferStruct,
-	FUniformBufferRHIParamRef SceneTexturesUniformBuffer,
+	FRHIUniformBuffer* SceneTexturesUniformBuffer,
 	EParticleSimulatePhase::Type Phase,
 	bool bUseFixDT)
 {
@@ -1889,7 +1889,7 @@ void InjectNewParticles(FRHICommandList& RHICmdList, FGraphicsPipelineStateIniti
 	}
 
 	const int32 MaxParticlesPerDrawCall = GParticleScratchVertexBufferSize / sizeof(FNewParticle);
-	FVertexBufferRHIParamRef ScratchVertexBufferRHI = GParticleScratchVertexBuffer.VertexBufferRHI;
+	FRHIVertexBuffer* ScratchVertexBufferRHI = GParticleScratchVertexBuffer.VertexBufferRHI;
 	int32 ParticleCount = NewParticles.Num();
 	int32 FirstParticle = 0;
 
@@ -2041,7 +2041,7 @@ public:
 	/**
 	 * Set parameters for this shader.
 	 */
-	void SetParameters(FRHICommandList& RHICmdList, int32 InVisualizationMode, FTexture2DRHIParamRef PositionTextureRHI, FTexture2DRHIParamRef CurveTextureRHI )
+	void SetParameters(FRHICommandList& RHICmdList, int32 InVisualizationMode, FRHITexture2D* PositionTextureRHI, FRHITexture2D* CurveTextureRHI )
 	{
 		FRHIPixelShader* PixelShader = GetPixelShader();
 		SetShaderValue(RHICmdList, PixelShader, VisualizationMode, InVisualizationMode );
@@ -2099,7 +2099,7 @@ static void VisualizeGPUSimulation(
 	int32 VisualizationMode,
 	FRenderTarget* RenderTarget,
 	const FParticleStateTextures& StateTextures,
-	FTexture2DRHIParamRef CurveTextureRHI
+	FRHITexture2D* CurveTextureRHI
 	)
 {
 	check(IsInRenderingThread());
@@ -2175,7 +2175,7 @@ static void VisualizeGPUSimulation(
  * @param VertexBuffer - The buffer with which to fill with particle indices.
  * @param InTiles - The list of tiles for which to generate indices.
  */
-static void BuildParticleVertexBuffer( FVertexBufferRHIParamRef VertexBufferRHI, const TArray<uint32>& InTiles )
+static void BuildParticleVertexBuffer(FRHIVertexBuffer* VertexBufferRHI, const TArray<uint32>& InTiles )
 {
 	check( IsInRenderingThread() );
 
@@ -2299,7 +2299,7 @@ public:
 		FRHICommandList& RHICmdList,
 		FParticleBoundsUniformBufferRef& UniformBuffer,
 		FRHIShaderResourceView* InIndicesSRV,
-		FTexture2DRHIParamRef PositionTextureRHI
+		FRHITexture2D* PositionTextureRHI
 		)
 	{
 		FRHIComputeShader* ComputeShaderRHI = GetComputeShader();
@@ -2365,7 +2365,7 @@ static bool AreBoundsValid( const FVector& Mins, const FVector& Maxs )
 static FBox ComputeParticleBounds(
 	FRHICommandListImmediate& RHICmdList,
 	FRHIShaderResourceView* VertexBufferSRV,
-	FTexture2DRHIParamRef PositionTextureRHI,
+	FRHITexture2D* PositionTextureRHI,
 	int32 ParticleCount )
 {
 	FBox BoundingBox;
@@ -4687,7 +4687,7 @@ void FFXSystem::PrepareGPUSimulation(FRHICommandListImmediate& RHICmdList)
 	FParticleStateTextures& CurrentStateTextures = ParticleSimulationResources->GetCurrentStateTextures();
 
 	// Setup render states.
-	FTextureRHIParamRef RenderTargets[2] =
+	FRHITexture* RenderTargets[2] =
 	{
 		CurrentStateTextures.PositionTextureTargetRHI,
 		CurrentStateTextures.VelocityTextureTargetRHI,
@@ -4702,7 +4702,7 @@ void FFXSystem::FinalizeGPUSimulation(FRHICommandListImmediate& RHICmdList)
 	FParticleStateTextures& CurrentStateTextures = ParticleSimulationResources->GetVisualizeStateTextures();
 
 	// Setup render states.
-	FTextureRHIParamRef RenderTargets[2] =
+	FRHITexture* RenderTargets[2] =
 	{
 		CurrentStateTextures.PositionTextureTargetRHI,
 		CurrentStateTextures.VelocityTextureTargetRHI,
@@ -4714,10 +4714,10 @@ void FFXSystem::FinalizeGPUSimulation(FRHICommandListImmediate& RHICmdList)
 void FFXSystem::SimulateGPUParticles(
 	FRHICommandListImmediate& RHICmdList,
 	EParticleSimulatePhase::Type Phase,
-	const FUniformBufferRHIParamRef ViewUniformBuffer,
+	FRHIUniformBuffer* ViewUniformBuffer,
 	const FGlobalDistanceFieldParameterData* GlobalDistanceFieldParameterData,
 	const FShaderParametersMetadata* SceneTexturesUniformBufferStruct,
-	FUniformBufferRHIParamRef SceneTexturesUniformBuffer
+	FRHIUniformBuffer* SceneTexturesUniformBuffer
 	)
 {
 	LLM_SCOPE(ELLMTag::Particles);
@@ -4733,8 +4733,8 @@ void FFXSystem::SimulateGPUParticles(
 	FParticleStateTextures& PrevStateTextures = ParticleSimulationResources->GetPreviousStateTextures();	
 
 	// Setup render states.
-	FTextureRHIParamRef CurrentStateRenderTargets[2] = { CurrentStateTextures.PositionTextureTargetRHI, CurrentStateTextures.VelocityTextureTargetRHI };
-	FTextureRHIParamRef PreviousStateRenderTargets[2] = { PrevStateTextures.PositionTextureTargetRHI, PrevStateTextures.VelocityTextureTargetRHI };
+	FRHITexture* CurrentStateRenderTargets[2] = { CurrentStateTextures.PositionTextureTargetRHI, CurrentStateTextures.VelocityTextureTargetRHI };
+	FRHITexture* PreviousStateRenderTargets[2] = { PrevStateTextures.PositionTextureTargetRHI, PrevStateTextures.VelocityTextureTargetRHI };
 	{
 		
 		// On some platforms, the textures are filled with garbage after creation, so we need to clear them to black the first time we use them
@@ -4787,7 +4787,7 @@ void FFXSystem::SimulateGPUParticles(
 		if (GPUSimulations.Num() > 0)
 		{
 			FVectorFieldUniformParameters VectorFieldParameters;
-			FTexture3DRHIParamRef BlackVolumeTextureRHI = (FTexture3DRHIParamRef)(FTextureRHIParamRef)GBlackVolumeTexture->TextureRHI;
+			FRHITexture3D* BlackVolumeTextureRHI = (FRHITexture3D*)(FRHITexture*)GBlackVolumeTexture->TextureRHI;
 			for (int32 Index = 0; Index < MAX_VECTOR_FIELDS; ++Index)
 			{
 				VectorFieldParameters.WorldToVolume[Index] = FMatrix::Identity;
@@ -4986,7 +4986,7 @@ void FFXSystem::SimulateGPUParticles(
 		SCOPED_GPU_STAT(RHICmdList, ParticleSimulation);
 
 		// Set render targets.
-		FTextureRHIParamRef InjectRenderTargets[4] =
+		FRHITexture* InjectRenderTargets[4] =
 		{
 			CurrentStateTextures.PositionTextureTargetRHI,
 			CurrentStateTextures.VelocityTextureTargetRHI,
@@ -5048,7 +5048,7 @@ void FFXSystem::SimulateGPUParticles(
 		//the fixed timestep works in two stages.  A first stage which simulates the fixed timestep and this second stage which simulates any remaining time from the actual delta time.  e.g.  fixed timestep of 16ms and actual dt of 23ms
 		//will make this second step simulate an interpolated extra 7ms.  This second interpolated step is what we render on THIS frame, but it is NOT fed into the next frame's simulation.  Thus we do not need to transfer it between GPUs in AFR mode.
 		FParticleStateTextures& VisualizeStateTextures = ParticleSimulationResources->GetPreviousStateTextures();
-		FTextureRHIParamRef VisualizeStateRHIs[2] = { VisualizeStateTextures.PositionTextureTargetRHI, VisualizeStateTextures.VelocityTextureTargetRHI };
+		FRHITexture* VisualizeStateRHIs[2] = { VisualizeStateTextures.PositionTextureTargetRHI, VisualizeStateTextures.VelocityTextureTargetRHI };
 
 		FRHIRenderPassInfo RPInfo(2, VisualizeStateRHIs, ERenderTargetActions::Load_Store);
 		{
@@ -5098,7 +5098,7 @@ void FFXSystem::UpdateMultiGPUResources(FRHICommandListImmediate& RHICmdList)
 		SCOPED_GPU_STAT(RHICmdList, ParticleSimulation);
 
 		// Set render targets.
-		FTextureRHIParamRef InjectRenderTargets[2] =
+		FRHITexture* InjectRenderTargets[2] =
 		{
 			ParticleSimulationResources->RenderAttributesTexture.TextureTargetRHI,
 			ParticleSimulationResources->SimulationAttributesTexture.TextureTargetRHI
diff --git a/Engine/Source/Runtime/Engine/Private/Particles/ParticleSortingGPU.cpp b/Engine/Source/Runtime/Engine/Private/Particles/ParticleSortingGPU.cpp
--- a/Engine/Source/Runtime/Engine/Private/Particles/ParticleSortingGPU.cpp
+++ b/Engine/Source/Runtime/Engine/Private/Particles/ParticleSortingGPU.cpp
@@ -124,7 +124,7 @@ public:
 	/**
 	 * Set the texture from which particle positions can be read.
 	 */
-	void SetPositionTextures(FRHICommandList& RHICmdList, FTexture2DRHIParamRef PositionTextureRHI)
+	void SetPositionTextures(FRHICommandList& RHICmdList, FRHITexture2D* PositionTextureRHI)
 	{
 		FRHIComputeShader* ComputeShaderRHI = GetComputeShader();
 		if (PositionTexture.IsBound())
@@ -179,7 +179,7 @@ static int32 GenerateParticleSortKeys(
 	FRHICommandListImmediate& RHICmdList,
 	FRHIUnorderedAccessView* KeyBufferUAV,
 	FRHIUnorderedAccessView* SortedVertexBufferUAV,
-	FTexture2DRHIParamRef PositionTextureRHI,
+	FRHITexture2D* PositionTextureRHI,
 	const TArray<FParticleSimulationSortInfo>& SimulationsToSort,
 	ERHIFeatureLevel::Type FeatureLevel
 	)
@@ -333,7 +333,7 @@ FGPUSortBuffers FParticleSortBuffers::GetSortBuffers()
 void SortParticlesGPU(
 	FRHICommandListImmediate& RHICmdList,
 	FParticleSortBuffers& ParticleSortBuffers,
-	FTexture2DRHIParamRef PositionTextureRHI,
+	FRHITexture2D* PositionTextureRHI,
 	const TArray<FParticleSimulationSortInfo>& SimulationsToSort,
 	ERHIFeatureLevel::Type FeatureLevel
 	)
@@ -346,7 +346,7 @@ void SortParticlesGPU(
 		const int32 StreamCount = 2;
 		for (int32 StreamIndex = 0; StreamIndex < StreamCount; ++StreamIndex)
 		{
-			RHICmdList.SetStreamSource(StreamIndex, FVertexBufferRHIParamRef(), 0);
+			RHICmdList.SetStreamSource(StreamIndex, nullptr, 0);
 		}
 	}
 
diff --git a/Engine/Source/Runtime/Engine/Private/Particles/ParticleSortingGPU.h b/Engine/Source/Runtime/Engine/Private/Particles/ParticleSortingGPU.h
--- a/Engine/Source/Runtime/Engine/Private/Particles/ParticleSortingGPU.h
+++ b/Engine/Source/Runtime/Engine/Private/Particles/ParticleSortingGPU.h
@@ -67,7 +67,7 @@ public:
 	/**
 	 * Retrieve the sorted vertex buffer that results will always be located at.
 	 */
-	FVertexBufferRHIParamRef GetSortedVertexBufferRHI(int32 BufferIndex)
+	FRHIVertexBuffer* GetSortedVertexBufferRHI(int32 BufferIndex)
 	{
 		check((BufferIndex & 0xFFFFFFFE) == 0);
 		return VertexBuffers[BufferIndex];
@@ -145,7 +145,7 @@ struct FParticleSimulationSortInfo
 void SortParticlesGPU(
 	FRHICommandListImmediate& RHICmdList,
 	FParticleSortBuffers& ParticleSortBuffers,
-	FTexture2DRHIParamRef PositionTextureRHI,
+	FRHITexture2D* PositionTextureRHI,
 	const TArray<FParticleSimulationSortInfo>& SimulationsToSort,
 	ERHIFeatureLevel::Type FeatureLevel
 	);
diff --git a/Engine/Source/Runtime/Engine/Private/SceneManagement.cpp b/Engine/Source/Runtime/Engine/Private/SceneManagement.cpp
--- a/Engine/Source/Runtime/Engine/Private/SceneManagement.cpp
+++ b/Engine/Source/Runtime/Engine/Private/SceneManagement.cpp
@@ -673,8 +673,8 @@ FViewUniformShaderParameters::FViewUniformShaderParameters()
 {
 	FMemory::Memzero(*this);
 
-	FTextureRHIParamRef BlackVolume = (GBlackVolumeTexture &&  GBlackVolumeTexture->TextureRHI) ? GBlackVolumeTexture->TextureRHI : GBlackTexture->TextureRHI; // for es2, this might need to be 2d
-	FTextureRHIParamRef BlackUintVolume = (GBlackUintVolumeTexture &&  GBlackUintVolumeTexture->TextureRHI) ? GBlackUintVolumeTexture->TextureRHI : GBlackTexture->TextureRHI; // for es2, this might need to be 2d
+	FRHITexture* BlackVolume = (GBlackVolumeTexture &&  GBlackVolumeTexture->TextureRHI) ? GBlackVolumeTexture->TextureRHI : GBlackTexture->TextureRHI; // for es2, this might need to be 2d
+	FRHITexture* BlackUintVolume = (GBlackUintVolumeTexture &&  GBlackUintVolumeTexture->TextureRHI) ? GBlackUintVolumeTexture->TextureRHI : GBlackTexture->TextureRHI; // for es2, this might need to be 2d
 	check(GBlackVolumeTexture);
 
 	MaterialTextureBilinearClampedSampler = TStaticSamplerState<SF_Bilinear, AM_Clamp, AM_Clamp, AM_Clamp>::GetRHI();
diff --git a/Engine/Source/Runtime/Engine/Private/Texture2D.cpp b/Engine/Source/Runtime/Engine/Private/Texture2D.cpp
--- a/Engine/Source/Runtime/Engine/Private/Texture2D.cpp
+++ b/Engine/Source/Runtime/Engine/Private/Texture2D.cpp
@@ -1708,7 +1708,7 @@ void FTexture2DResource::ReleaseRHI()
 
 	FTextureResource::ReleaseRHI();
 	Texture2DRHI.SafeRelease();
-	RHIUpdateTextureReference(Owner->TextureReference.TextureReferenceRHI,FTextureRHIParamRef());
+	RHIUpdateTextureReference(Owner->TextureReference.TextureReferenceRHI, nullptr);
 }
 
 void FTexture2DResource::CreateSamplerStates(float MipMapBias)
diff --git a/Engine/Source/Runtime/Engine/Private/Texture2DDynamic.cpp b/Engine/Source/Runtime/Engine/Private/Texture2DDynamic.cpp
--- a/Engine/Source/Runtime/Engine/Private/Texture2DDynamic.cpp
+++ b/Engine/Source/Runtime/Engine/Private/Texture2DDynamic.cpp
@@ -70,7 +70,7 @@ void FTexture2DDynamicResource::InitRHI()
 /** Called when the resource is released. This is only called by the rendering thread. */
 void FTexture2DDynamicResource::ReleaseRHI()
 {
-	RHIUpdateTextureReference(Owner->TextureReference.TextureReferenceRHI,FTextureRHIParamRef());
+	RHIUpdateTextureReference(Owner->TextureReference.TextureReferenceRHI, nullptr);
 	FTextureResource::ReleaseRHI();
 	Texture2DRHI.SafeRelease();
 }
diff --git a/Engine/Source/Runtime/Engine/Private/UnrealEngine.cpp b/Engine/Source/Runtime/Engine/Private/UnrealEngine.cpp
--- a/Engine/Source/Runtime/Engine/Private/UnrealEngine.cpp
+++ b/Engine/Source/Runtime/Engine/Private/UnrealEngine.cpp
@@ -2862,7 +2862,7 @@ public:
 	{
 	}
 
-	virtual void RenderTexture_RenderThread(FRHICommandListImmediate& RHICmdList, FTexture2DRHIParamRef BackBuffer, FTexture2DRHIParamRef SrcTexture, FVector2D WindowSize) const override
+	virtual void RenderTexture_RenderThread(FRHICommandListImmediate& RHICmdList, FRHITexture2D* BackBuffer, FRHITexture2D* SrcTexture, FVector2D WindowSize) const override
 	{
 		check(IsInRenderingThread());
 
